Refactor of the existing DDC Toolbar Widget

Moved all the code to a plugin module
Removed superfluous icons from the toolbar
Added combo-box drop down menu options
Added dockable window support
Added icons to the starship style
Numerous UI layout tweaks

diff --git a/Engine/Source/Editor/DerivedDataEditor/DerivedDataEditor.Build.cs b/Engine/Source/Editor/DerivedDataEditor/DerivedDataEditor.Build.cs
new file mode 100644
--- /dev/null
+++ b/Engine/Source/Editor/DerivedDataEditor/DerivedDataEditor.Build.cs
@@ -0,0 +1,42 @@
+// Copyright Epic Games, Inc. All Rights Reserved.
+
+using UnrealBuildTool;
+
+public class DerivedDataEditor : ModuleRules
+{
+	public DerivedDataEditor(ReadOnlyTargetRules Target)
+		 : base(Target)
+	{
+		PublicIncludePaths.Add(ModuleDirectory + "/Public");
+
+		PrivateDependencyModuleNames.AddRange(
+			new string[] {
+				"Core",
+				"CoreUObject",
+				"Engine",
+				"Slate",
+				"SlateCore",
+				"InputCore",
+				"EditorStyle",
+				"EditorFramework",
+				"UnrealEd",
+				"ToolMenus",
+				"OutputLog",
+				"DerivedDataCache",
+				"EditorSubsystem",
+				"WorkspaceMenuStructure",
+				"MessageLog",
+			});
+
+		DynamicallyLoadedModuleNames.AddRange(
+			new string[] {
+				"MainFrame",
+			});
+
+
+		PrivateIncludePathModuleNames.AddRange(
+			new string[] {
+				"MainFrame",
+			});
+	}
+}
diff --git a/Engine/Source/Editor/DerivedDataEditor/Private/DerivedDataEditorModule.cpp b/Engine/Source/Editor/DerivedDataEditor/Private/DerivedDataEditorModule.cpp
new file mode 100644
--- /dev/null
+++ b/Engine/Source/Editor/DerivedDataEditor/Private/DerivedDataEditorModule.cpp
@@ -0,0 +1,160 @@
+// Copyright Epic Games, Inc. All Rights Reserved.
+
+#include "DerivedDataEditorModule.h"
+#include "DerivedDataInformation.h"
+#include "SDerivedDataStatusBar.h"
+#include "SDerivedDataDialogs.h"
+#include "SDerivedDataCacheSettings.h"
+#include "Widgets/DeclarativeSyntaxSupport.h"
+#include "Modules/ModuleManager.h"
+#include "Widgets/Docking/SDockTab.h"
+
+#include "WorkspaceMenuStructure.h"
+#include "WorkspaceMenuStructureModule.h"
+
+#define LOCTEXT_NAMESPACE "DerivedDataEditor"
+
+IMPLEMENT_MODULE(FDerivedDataEditorModule, DerivedDataEditor );
+
+static const FName DerivedDataResourceUsageTabName = FName(TEXT("DerivedDataResourceUsage"));
+static const FName DerivedDataCacheStatisticsTabName = FName(TEXT("DerivedDataCacheStatistics"));
+
+void FDerivedDataEditorModule::StartupModule()
+{
+	// We're going to call a static function in the editor style module, so we need to make sure the module has actually been loaded
+	FModuleManager::Get().LoadModuleChecked("EditorStyle");
+
+	const FSlateIcon ResourceUsageIcon(FEditorStyle::GetStyleSetName(), "DerivedData.ResourceUsage");
+
+	FGlobalTabmanager::Get()->RegisterNomadTabSpawner(DerivedDataResourceUsageTabName, FOnSpawnTab::CreateRaw(this, &FDerivedDataEditorModule::CreateResourceUsageTab))
+		.SetDisplayName(LOCTEXT("DerivedDataResourceUsageTabTitle", "Resource Usage"))
+		.SetTooltipText(LOCTEXT("DerivedDataResourceUsageTabToolTipText", "Derived Data Resource Usage"))
+		.SetGroup(WorkspaceMenu::GetMenuStructure().GetToolsCategory())
+		.SetIcon(ResourceUsageIcon);
+
+	const FSlateIcon CacheStatisticsIcon(FEditorStyle::GetStyleSetName(), "DerivedData.CacheStatistics");
+
+	FGlobalTabmanager::Get()->RegisterNomadTabSpawner(DerivedDataCacheStatisticsTabName, FOnSpawnTab::CreateRaw(this, &FDerivedDataEditorModule::CreateCacheStatisticsTab))
+		.SetDisplayName(LOCTEXT("DerivedDataCacheStatisticsTabTitle", "Cache Statistics"))
+		.SetTooltipText(LOCTEXT("DerivedDataCacheStatisticsTabToolTipText", "Derived Data Cache Statistics"))
+		.SetGroup(WorkspaceMenu::GetMenuStructure().GetToolsCategory())
+		.SetIcon(CacheStatisticsIcon);
+
+
+#if WITH_RELOAD
+	// This code attempts to relaunch the tabs when you reload this module
+	if (IsReloadActive() && FSlateApplication::IsInitialized())
+	{
+		ShowCacheStatisticsTab();
+		ShowResourceUsageTab();
+	}
+#endif // WITH_RELOAD
+
+	FDerivedDataStatusBarMenuCommands::Register();
+}
+
+void FDerivedDataEditorModule::ShutdownModule()
+{
+	if (FSlateApplication::IsInitialized())
+	{
+		FGlobalTabmanager::Get()->UnregisterNomadTabSpawner(DerivedDataResourceUsageTabName);
+
+		if (ResourceUsageTab.IsValid())
+		{
+			ResourceUsageTab.Pin()->RequestCloseTab();
+		}
+
+		FGlobalTabmanager::Get()->UnregisterNomadTabSpawner(DerivedDataCacheStatisticsTabName);
+
+		if (CacheStatisticsTab.IsValid())
+		{
+			CacheStatisticsTab.Pin()->RequestCloseTab();
+		}
+	}
+
+	FDerivedDataStatusBarMenuCommands::Unregister();
+}
+
+TSharedRef<SWidget> FDerivedDataEditorModule::CreateStatusBarWidget()
+{
+	return SNew(SDerivedDataStatusBarWidget);
+}
+
+TSharedPtr<SWidget> FDerivedDataEditorModule::CreateResourceUsageDialog()
+{
+	return SNew(SDerivedDataResourceUsageDialog);
+}
+
+TSharedPtr<SWidget> FDerivedDataEditorModule::CreateCacheStatisticsDialog()
+{
+	return SNew(SDerivedDataCacheStatisticsDialog);
+}
+
+void FDerivedDataEditorModule::ShowResourceUsageTab()
+{
+	FGlobalTabmanager::Get()->TryInvokeTab(FTabId(DerivedDataResourceUsageTabName));
+}
+
+
+void FDerivedDataEditorModule::ShowCacheStatisticsTab()
+{
+	FGlobalTabmanager::Get()->TryInvokeTab(FTabId(DerivedDataCacheStatisticsTabName));
+}
+
+TSharedRef<SDockTab> FDerivedDataEditorModule::CreateResourceUsageTab(const FSpawnTabArgs& Args)
+{
+	return SAssignNew(ResourceUsageTab, SDockTab)
+	.TabRole(ETabRole::NomadTab)
+	[
+		CreateResourceUsageDialog().ToSharedRef()
+	];
+}
+
+TSharedRef<SDockTab> FDerivedDataEditorModule::CreateCacheStatisticsTab(const FSpawnTabArgs& Args)
+{
+	return SAssignNew(CacheStatisticsTab, SDockTab)
+		.TabRole(ETabRole::NomadTab)
+		[
+			CreateCacheStatisticsDialog().ToSharedRef()
+		];
+}
+
+void FDerivedDataEditorModule::ShowSettingsDialog()
+{
+	if (SettingsWindow.IsValid())
+	{
+		SettingsWindow->BringToFront();
+	}
+	else
+	{
+		// Create the window
+		SettingsWindow = SNew(SWindow)
+			.Title(LOCTEXT("DerrivedDataCacheSettingsWindowTitle", "Cache Settings"))
+			.HasCloseButton(true)
+			.SupportsMaximize(false)
+			.SupportsMinimize(true)
+			.SizingRule(ESizingRule::Autosized);
+
+		// Set the closed callback
+		SettingsWindow->SetOnWindowClosed(FOnWindowClosed::CreateRaw(this,&FDerivedDataEditorModule::OnSettingsDialogClosed));
+
+		// Setup the content for the created window.
+		SettingsWindow->SetContent(SAssignNew(SettingsDialog, SDerivedDataCacheSettingsDialog));
+
+		TSharedPtr<SWindow> RootWindow = FGlobalTabmanager::Get()->GetRootWindow();
+		if (RootWindow.IsValid())
+		{
+			FSlateApplication::Get().AddWindowAsNativeChild(SettingsWindow.ToSharedRef(), RootWindow.ToSharedRef());
+		}
+		else
+		{
+			FSlateApplication::Get().AddWindow(SettingsWindow.ToSharedRef());
+		}
+	}
+}
+
+void FDerivedDataEditorModule::OnSettingsDialogClosed(const TSharedRef<SWindow>& InWindow)
+{
+	SettingsWindow = nullptr;
+	SettingsDialog = nullptr;
+}
\ No newline at end of file
diff --git a/Engine/Source/Editor/DerivedDataEditor/Private/DerivedDataInformation.cpp b/Engine/Source/Editor/DerivedDataEditor/Private/DerivedDataInformation.cpp
new file mode 100644
--- /dev/null
+++ b/Engine/Source/Editor/DerivedDataEditor/Private/DerivedDataInformation.cpp
@@ -0,0 +1,154 @@
+// Copyright Epic Games, Inc. All Rights Reserved.
+
+#include "DerivedDataInformation.h"
+#include "SDerivedDataStatusBar.h"
+#include "DerivedDataCacheUsageStats.h"
+#include "DerivedDataCacheInterface.h"
+#include "DerivedDataBackendInterface.h"
+#include "ToolMenus.h"
+
+#define LOCTEXT_NAMESPACE "DerivedDataEditor"
+
+double FDerivedDataInformation::GetCacheActivitySizeBytes(bool bGet, bool bLocal)
+{
+	PRAGMA_DISABLE_DEPRECATION_WARNINGS
+		TSharedRef<FDerivedDataCacheStatsNode> RootUsage = GetDerivedDataCache()->GatherUsageStats();
+	PRAGMA_ENABLE_DEPRECATION_WARNINGS
+		TArray<TSharedRef<const FDerivedDataCacheStatsNode>> LeafUsageStats;
+	RootUsage->ForEachDescendant([&LeafUsageStats](TSharedRef<const FDerivedDataCacheStatsNode> Node) {
+		if (Node->Children.Num() == 0)
+		{
+			LeafUsageStats.Add(Node);
+		}
+		});
+
+	int64 TotalBytes = 0;
+
+	for (int32 Index = 0; Index < LeafUsageStats.Num(); Index++)
+	{
+		const FDerivedDataBackendInterface* Backend = LeafUsageStats[Index]->GetBackendInterface();
+
+		if ((Backend->GetSpeedClass() == FDerivedDataBackendInterface::ESpeedClass::Local) != bLocal)
+			continue;
+
+		TSharedRef<FDerivedDataCacheStatsNode> Usage = Backend->GatherUsageStats();
+
+		for (const auto& KVP : Usage->Stats)
+		{
+			const FDerivedDataCacheUsageStats& Stats = KVP.Value;
+
+			if (bGet)
+			{
+				TotalBytes += Stats.GetStats.GetAccumulatedValueAnyThread(FCookStats::CallStats::EHitOrMiss::Hit, FCookStats::CallStats::EStatType::Bytes);
+			}
+			else
+			{
+				TotalBytes += Stats.PutStats.GetAccumulatedValueAnyThread(FCookStats::CallStats::EHitOrMiss::Hit, FCookStats::CallStats::EStatType::Bytes);
+			}
+		}
+	}
+
+	return TotalBytes;
+}
+
+
+double FDerivedDataInformation::GetCacheActivityTimeSeconds(bool bGet, bool bLocal)
+{
+	PRAGMA_DISABLE_DEPRECATION_WARNINGS
+		TSharedRef<FDerivedDataCacheStatsNode> RootUsage = GetDerivedDataCache()->GatherUsageStats();
+	PRAGMA_ENABLE_DEPRECATION_WARNINGS
+		TArray<TSharedRef<const FDerivedDataCacheStatsNode>> LeafUsageStats;
+	RootUsage->ForEachDescendant([&LeafUsageStats](TSharedRef<const FDerivedDataCacheStatsNode> Node) {
+		if (Node->Children.Num() == 0)
+		{
+			LeafUsageStats.Add(Node);
+		}
+		});
+
+	int64 TotalCycles = 0;
+
+	for (int32 Index = 0; Index < LeafUsageStats.Num(); Index++)
+	{
+		const FDerivedDataBackendInterface* Backend = LeafUsageStats[Index]->GetBackendInterface();
+
+		if ((Backend->GetSpeedClass() == FDerivedDataBackendInterface::ESpeedClass::Local) != bLocal)
+			continue;
+
+		TSharedRef<FDerivedDataCacheStatsNode> Usage = Backend->GatherUsageStats();
+
+		for (const auto& KVP : Usage->Stats)
+		{
+			const FDerivedDataCacheUsageStats& Stats = KVP.Value;
+
+			if (bGet)
+			{
+				TotalCycles +=
+					(Stats.GetStats.GetAccumulatedValueAnyThread(FCookStats::CallStats::EHitOrMiss::Hit, FCookStats::CallStats::EStatType::Cycles) +
+						Stats.GetStats.GetAccumulatedValueAnyThread(FCookStats::CallStats::EHitOrMiss::Miss, FCookStats::CallStats::EStatType::Cycles));
+
+				TotalCycles +=
+					(Stats.PrefetchStats.GetAccumulatedValueAnyThread(FCookStats::CallStats::EHitOrMiss::Hit, FCookStats::CallStats::EStatType::Cycles) +
+						Stats.PrefetchStats.GetAccumulatedValueAnyThread(FCookStats::CallStats::EHitOrMiss::Miss, FCookStats::CallStats::EStatType::Cycles));
+			}
+			else
+			{
+				TotalCycles +=
+					(Stats.PutStats.GetAccumulatedValueAnyThread(FCookStats::CallStats::EHitOrMiss::Hit, FCookStats::CallStats::EStatType::Cycles) +
+						Stats.PutStats.GetAccumulatedValueAnyThread(FCookStats::CallStats::EHitOrMiss::Miss, FCookStats::CallStats::EStatType::Cycles));
+			}
+		}
+	}
+
+	return (double)TotalCycles * FPlatformTime::GetSecondsPerCycle();
+}
+
+bool FDerivedDataInformation::GetHasLocalCache()
+{
+	PRAGMA_DISABLE_DEPRECATION_WARNINGS
+		TSharedRef<FDerivedDataCacheStatsNode> RootUsage = GetDerivedDataCache()->GatherUsageStats();
+	PRAGMA_ENABLE_DEPRECATION_WARNINGS
+		TArray<TSharedRef<const FDerivedDataCacheStatsNode>> LeafUsageStats;
+	RootUsage->ForEachDescendant([&LeafUsageStats](TSharedRef<const FDerivedDataCacheStatsNode> Node) {
+		if (Node->Children.Num() == 0)
+		{
+			LeafUsageStats.Add(Node);
+		}
+		});
+
+	for (int32 Index = 0; Index < LeafUsageStats.Num(); Index++)
+	{
+		const FDerivedDataBackendInterface* Backend = LeafUsageStats[Index]->GetBackendInterface();
+
+		if (Backend->GetSpeedClass() == FDerivedDataBackendInterface::ESpeedClass::Local)
+			return true;
+	}
+
+	return false;
+}
+
+
+bool FDerivedDataInformation::GetHasRemoteCache()
+{
+	PRAGMA_DISABLE_DEPRECATION_WARNINGS
+		TSharedRef<FDerivedDataCacheStatsNode> RootUsage = GetDerivedDataCache()->GatherUsageStats();
+	PRAGMA_ENABLE_DEPRECATION_WARNINGS
+		TArray<TSharedRef<const FDerivedDataCacheStatsNode>> LeafUsageStats;
+	RootUsage->ForEachDescendant([&LeafUsageStats](TSharedRef<const FDerivedDataCacheStatsNode> Node) {
+		if (Node->Children.Num() == 0)
+		{
+			LeafUsageStats.Add(Node);
+		}
+		});
+
+	for (int32 Index = 0; Index < LeafUsageStats.Num(); Index++)
+	{
+		const FDerivedDataBackendInterface* Backend = LeafUsageStats[Index]->GetBackendInterface();
+
+		if (Backend->GetSpeedClass() != FDerivedDataBackendInterface::ESpeedClass::Local)
+			return true;
+	}
+
+	return false;
+}
+
+#undef LOCTEXT_NAMESPACE
\ No newline at end of file
diff --git a/Engine/Source/Editor/DerivedDataEditor/Private/DerivedDataInformation.h b/Engine/Source/Editor/DerivedDataEditor/Private/DerivedDataInformation.h
new file mode 100644
--- /dev/null
+++ b/Engine/Source/Editor/DerivedDataEditor/Private/DerivedDataInformation.h
@@ -0,0 +1,16 @@
+// Copyright Epic Games, Inc. All Rights Reserved.
+
+#pragma once
+
+#include "CoreMinimal.h"
+#include "Framework/Commands/Commands.h"
+
+class FDerivedDataInformation
+{
+public:
+
+	static double	GetCacheActivityTimeSeconds(bool bGet, bool bLocal);
+	static double	GetCacheActivitySizeBytes(bool bGet, bool bLocal);
+	static bool		GetHasLocalCache();
+	static bool		GetHasRemoteCache();
+};
diff --git a/Engine/Source/Editor/DerivedDataEditor/Private/SDerivedDataCacheSettings.cpp b/Engine/Source/Editor/DerivedDataEditor/Private/SDerivedDataCacheSettings.cpp
new file mode 100644
--- /dev/null
+++ b/Engine/Source/Editor/DerivedDataEditor/Private/SDerivedDataCacheSettings.cpp
@@ -0,0 +1,243 @@
+// Copyright Epic Games, Inc. All Rights Reserved.
+
+#include "SDerivedDataCacheSettings.h"
+#include "DerivedDataInformation.h"
+#include "DerivedDataCacheInterface.h"
+#include "Widgets/SBoxPanel.h"
+#include "Widgets/Text/STextBlock.h"
+#include "Widgets/Layout/SGridPanel.h"
+#include "Widgets/Images/SImage.h"
+#include "Widgets/Layout/SExpandableArea.h"
+#include "Widgets/Images/SThrobber.h"
+
+#include "Logging/MessageLog.h"
+#include "MessageLogModule.h"
+
+#define LOCTEXT_NAMESPACE "DerivedDataEditor"
+
+void SDerivedDataCacheSettingsDialog::Construct(const FArguments& InArgs)
+{
+	FMessageLogModule& MessageLogModule = FModuleManager::LoadModuleChecked<FMessageLogModule>("MessageLog");
+	TSharedRef<class IMessageLogListing> MessageLogListing = MessageLogModule.GetLogListing(TEXT("Config"));
+
+	this->ChildSlot
+	[
+		SNew(SBox)
+		.WidthOverride(480.f)
+		[
+			SNew(SBorder)
+			.Padding(10.f)
+			.BorderImage(FAppStyle::Get().GetBrush("Brushes.Panel"))
+			[
+				SNew(SVerticalBox)
+				+ SVerticalBox::Slot()
+				.AutoHeight()
+				/*.Expose(GridSlot)
+				[
+					GetGridPanel()
+				]*/
+
+				+ SVerticalBox::Slot()
+				.AutoHeight()
+				[
+					SNew(SHorizontalBox)
+					+SHorizontalBox::Slot()
+					.FillWidth(1.0f)
+					.Padding(FMargin(0.0f, 0.0f, 16.0f, 0.0f))
+					[
+						SNew(SVerticalBox)
+						+SVerticalBox::Slot()
+						.Padding(FMargin(0.0f, 0.0f, 0.0f, 10.0f))
+						.VAlign(VAlign_Center)
+						.HAlign(HAlign_Right)
+						[
+							SNew(STextBlock)
+							.Text(LOCTEXT("Setting1Label", "Setting1"))
+							.ToolTipText( LOCTEXT("Setting1Label_Tooltip", "Setting1") )
+						]
+						+SVerticalBox::Slot()
+						.Padding(FMargin(0.0f, 0.0f, 0.0f, 10.0f))
+						.VAlign(VAlign_Center)
+						.HAlign(HAlign_Right)
+						[
+							SNew(STextBlock)
+							.Text(LOCTEXT("Setting2Label", "Setting2"))
+							.ToolTipText( LOCTEXT("Setting2Label_Tooltip", "Setting2") )
+						]		
+					]
+					+SHorizontalBox::Slot()
+					.FillWidth(2.0f)
+					[
+						SNew(SVerticalBox)
+						+SVerticalBox::Slot()
+						.VAlign(VAlign_Center)
+						.Padding(FMargin(0.0f, 0.0f, 0.0f, 10.0f))
+						[
+							SNew(SEditableTextBox)
+							.Text(this, &SDerivedDataCacheSettingsDialog::GetSetting1Text)
+							.ToolTipText( LOCTEXT("Setting1Label_Tooltip", "Setting1") )
+							.OnTextCommitted(this, &SDerivedDataCacheSettingsDialog::OnSetting2TextCommited)
+							.OnTextChanged(this, &SDerivedDataCacheSettingsDialog::OnSetting2TextCommited, ETextCommit::Default)
+						]
+						+SVerticalBox::Slot()
+						.Padding(FMargin(0.0f, 0.0f, 0.0f, 10.0f))
+						.VAlign(VAlign_Center)
+						[
+							SNew(SEditableTextBox)
+							.Text(this, &SDerivedDataCacheSettingsDialog::GetSetting2Text)
+							.ToolTipText( LOCTEXT("Setting2Label_Tooltip", "Setting2") )
+							.OnTextCommitted(this, &SDerivedDataCacheSettingsDialog::OnSetting2TextCommited)
+							.OnTextChanged(this, &SDerivedDataCacheSettingsDialog::OnSetting2TextCommited, ETextCommit::Default)
+						]		
+					]	
+				]
+
+				//+ SVerticalBox::Slot()
+				//.VAlign(VAlign_Top)
+				//.AutoHeight()
+				//.Padding(0.0f)
+				//[
+				//	SNew(SExpandableArea)
+				//	.BorderImage(FEditorStyle::GetBrush("NoBorder"))
+				//	.InitiallyCollapsed(true)
+				//	.HeaderContent()
+				//	[
+				//		SNew(STextBlock)
+				//		.Text(LOCTEXT("LogTitle", "Derived Data Cache Log"))
+				//	]
+				//	.BodyContent()
+				//	[
+				//		SNew(SBorder)
+				//		.BorderImage(FAppStyle::Get().GetBrush("Brushes.Recessed"))
+				//		.Padding(0)
+				//		[
+				//			SNew(SBox)
+				//			.HeightOverride(250)
+				//			[
+				//				MessageLogModule.CreateLogListingWidget(MessageLogListing)
+				//			]
+				//		]
+				//	]
+				//]
+
+				+SVerticalBox::Slot()
+				.VAlign(VAlign_Bottom)
+				.Padding(8.0f, 16.0f, 8.0f, 16.0f)
+				[
+					SNew(SHorizontalBox)
+					+SHorizontalBox::Slot()
+					.AutoWidth()
+					.VAlign(VAlign_Center)
+					.HAlign(HAlign_Left)
+					[
+						SNew(SThrobber)
+						.Visibility(this, &SDerivedDataCacheSettingsDialog::GetThrobberVisibility)
+					]
+					+SHorizontalBox::Slot()
+					.FillWidth(1.0f)
+					.HAlign(HAlign_Right)
+					[
+						SNew(SHorizontalBox)
+						+ SHorizontalBox::Slot()
+						.AutoWidth()
+						.Padding(FMargin(5.0f, 0.0f))
+						[
+							SNew(SButton)
+							.VAlign(VAlign_Center)
+							.ButtonStyle(FAppStyle::Get(), "PrimaryButton")
+  							.TextStyle( FAppStyle::Get(), "DialogButtonText" )
+							.Text(LOCTEXT("AcceptSettings", "Accept Settings"))
+							.OnClicked( this, &SDerivedDataCacheSettingsDialog::OnAcceptSettings )
+							.IsEnabled( this, &SDerivedDataCacheSettingsDialog::IsAcceptSettingsEnabled )
+						]
+						+ SHorizontalBox::Slot()
+						.AutoWidth()
+						.Padding(FMargin(5.0f, 0.0f))
+						[
+							SNew(SButton)
+							.VAlign(VAlign_Center)
+  							.TextStyle( FAppStyle::Get(), "DialogButtonText" )
+							.Text(LOCTEXT("RunWithoutCache", "Disable Derived Data Cache"))
+							.OnClicked( this, &SDerivedDataCacheSettingsDialog::OnDisableDerivedDataCache )
+							.IsEnabled( this, &SDerivedDataCacheSettingsDialog::IsDerivedDataCacheEnabled)
+						]
+					]
+				]
+			]
+		]
+	];
+
+	RegisterActiveTimer(0.5f, FWidgetActiveTimerDelegate::CreateSP(this, &SDerivedDataCacheSettingsDialog::UpdateGridPanels));
+}
+
+EActiveTimerReturnType SDerivedDataCacheSettingsDialog::UpdateGridPanels(double InCurrentTime, float InDeltaTime)
+{
+	/*(*GridSlot)
+	[
+		GetGridPanel()
+	];*/
+
+	SlatePrepass(GetPrepassLayoutScaleMultiplier());
+
+	return EActiveTimerReturnType::Continue;
+}
+
+TSharedRef<SWidget> SDerivedDataCacheSettingsDialog::GetGridPanel()
+{
+	TSharedRef<SGridPanel> Panel =
+		SNew(SGridPanel);
+
+	return Panel;
+}
+
+
+FText SDerivedDataCacheSettingsDialog::GetSetting1Text() const
+{
+	return LOCTEXT("Setting1Text", "Setting1");
+}
+
+void SDerivedDataCacheSettingsDialog::OnSetting1TextCommited(const FText& InText, ETextCommit::Type InCommitType) const
+{
+	
+}
+
+
+FText SDerivedDataCacheSettingsDialog::GetSetting2Text() const
+{
+	return LOCTEXT("Setting2Text", "Setting2");
+}
+
+void SDerivedDataCacheSettingsDialog::OnSetting2TextCommited(const FText& InText, ETextCommit::Type InCommitType) const
+{
+
+}
+
+
+EVisibility SDerivedDataCacheSettingsDialog::GetThrobberVisibility() const
+{
+	return EVisibility::Hidden;
+}
+
+bool SDerivedDataCacheSettingsDialog::IsDerivedDataCacheEnabled() const
+{
+	return true;
+}
+
+bool SDerivedDataCacheSettingsDialog::IsAcceptSettingsEnabled() const
+{
+	return true;
+}
+
+FReply SDerivedDataCacheSettingsDialog::OnAcceptSettings()
+{
+
+	return FReply::Handled();
+}
+
+FReply SDerivedDataCacheSettingsDialog::OnDisableDerivedDataCache()
+{
+	return FReply::Handled();
+}
+
+
+#undef LOCTEXT_NAMESPACE
\ No newline at end of file
diff --git a/Engine/Source/Editor/DerivedDataEditor/Private/SDerivedDataCacheSettings.h b/Engine/Source/Editor/DerivedDataEditor/Private/SDerivedDataCacheSettings.h
new file mode 100644
--- /dev/null
+++ b/Engine/Source/Editor/DerivedDataEditor/Private/SDerivedDataCacheSettings.h
@@ -0,0 +1,37 @@
+// Copyright Epic Games, Inc. All Rights Reserved.
+
+#pragma once
+
+#include "CoreMinimal.h"
+#include "Widgets/DeclarativeSyntaxSupport.h"
+#include "Widgets/SWidget.h"
+#include "Widgets/SCompoundWidget.h"
+
+class SDerivedDataCacheSettingsDialog : public SCompoundWidget
+{
+	SLATE_BEGIN_ARGS(SDerivedDataCacheSettingsDialog) {}
+	SLATE_END_ARGS()
+
+	void Construct(const FArguments& InArgs);
+
+private:
+
+	TSharedRef<SWidget> GetGridPanel();
+
+	EActiveTimerReturnType UpdateGridPanels(double InCurrentTime, float InDeltaTime);
+
+	SVerticalBox::FSlot* GridSlot = nullptr;
+
+	FText GetSetting1Text() const;
+	FText GetSetting2Text() const;
+
+	void OnSetting1TextCommited(const FText& InText, ETextCommit::Type InCommitType) const;
+	void OnSetting2TextCommited(const FText& InText, ETextCommit::Type InCommitType) const;
+
+	EVisibility GetThrobberVisibility() const;
+	bool IsDerivedDataCacheEnabled() const;
+	bool IsAcceptSettingsEnabled() const;
+	
+	FReply OnAcceptSettings();	
+	FReply OnDisableDerivedDataCache();
+};
diff --git a/Engine/Source/Editor/DerivedDataEditor/Private/SDerivedDataDialogs.cpp b/Engine/Source/Editor/DerivedDataEditor/Private/SDerivedDataDialogs.cpp
new file mode 100644
--- /dev/null
+++ b/Engine/Source/Editor/DerivedDataEditor/Private/SDerivedDataDialogs.cpp
@@ -0,0 +1,727 @@
+// Copyright Epic Games, Inc. All Rights Reserved.
+
+#include "SDerivedDataDialogs.h"
+#include "DerivedDataInformation.h"
+#include "DerivedDataCacheInterface.h"
+#include "DerivedDataBackendInterface.h"
+#include "DerivedDataCacheUsageStats.h"
+#include "Widgets/SBoxPanel.h"
+#include "Widgets/SToolTip.h"
+#include "Widgets/Text/STextBlock.h"
+#include "Widgets/Layout/SGridPanel.h"
+#include "Widgets/Images/SImage.h"
+#include "Styling/StyleColors.h"
+#include "Math/UnitConversion.h"
+#include "Framework/Application/SlateApplication.h"
+#include "Internationalization/FastDecimalFormat.h"
+
+
+#define LOCTEXT_NAMESPACE "DerivedDataCacheEditor"
+
+static FString SingleDecimalFormat(double Value)
+{
+	const FNumberFormattingOptions NumberFormattingOptions = FNumberFormattingOptions()
+		.SetUseGrouping(true)
+		.SetMinimumFractionalDigits(1)
+		.SetMaximumFractionalDigits(1);
+	return FastDecimalFormat::NumberToString(Value, ExpressionParser::GetLocalizedNumberFormattingRules(), NumberFormattingOptions);
+};
+
+void SDerivedDataRemoteStoreDialog::Construct(const FArguments& InArgs)
+{
+	this->ChildSlot
+	[
+		SNew(SVerticalBox)
+
+		+ SVerticalBox::Slot()
+		.AutoHeight()
+		.Padding(0, 20, 0, 0)
+		.Expose(GridSlot)
+		[
+			GetGridPanel()
+		]
+	];
+
+	RegisterActiveTimer(0.5f, FWidgetActiveTimerDelegate::CreateSP(this, &SDerivedDataRemoteStoreDialog::UpdateGridPanels));
+}
+
+EActiveTimerReturnType SDerivedDataRemoteStoreDialog::UpdateGridPanels(double InCurrentTime, float InDeltaTime)
+{
+	(*GridSlot)
+	[
+		GetGridPanel()
+	];
+
+	SlatePrepass(GetPrepassLayoutScaleMultiplier());
+
+	return EActiveTimerReturnType::Continue;
+}
+
+TSharedRef<SWidget> SDerivedDataRemoteStoreDialog::GetGridPanel()
+{
+	TArray<FDerivedDataCacheResourceStat> DDCResourceStats;
+
+	// Grab the latest resource stats
+	GetDerivedDataCacheRef().GatherResourceStats(DDCResourceStats);
+
+	FDerivedDataCacheResourceStat DDCResourceStatsTotal(TEXT("Total"));
+
+	// Accumulate Totals
+	for (const FDerivedDataCacheResourceStat& Stat : DDCResourceStats)
+	{
+		DDCResourceStatsTotal.Accumulate(Stat);
+	}
+
+	const int64 TotalCount = DDCResourceStatsTotal.LoadCount + DDCResourceStatsTotal.BuildCount;
+	const double Efficiency = TotalCount > 0 ? (double)DDCResourceStatsTotal.LoadCount / TotalCount : 0.0;
+
+	const double DownloadedBytesMB = FUnitConversion::Convert(FDerivedDataInformation::GetCacheActivitySizeBytes(true, false), EUnit::Bytes, EUnit::Megabytes);
+	const double UploadedBytesMB = FUnitConversion::Convert(FDerivedDataInformation::GetCacheActivitySizeBytes(false, false), EUnit::Bytes, EUnit::Megabytes);
+	
+	TSharedRef<SGridPanel> Panel =
+		SNew(SGridPanel);
+	
+	int32 Row = 0;
+
+	/*Panel->AddSlot(0, Row)
+		[
+			SNew(STextBlock)
+			.ColorAndOpacity(FStyleColors::Foreground)
+			.Text(LOCTEXT("Efficiency", "Efficiency"))
+		];
+
+	Panel->AddSlot(1, Row)
+		[
+			SNew(STextBlock)
+			.Margin(FMargin(5, 0))
+			.Text(FText::FromString(SingleDecimalFormat(Efficiency * 100.0f) + TEXT(" %")))
+			.Justification(ETextJustify::Right)
+		];
+
+	Row++;
+
+	Panel->AddSlot(0, Row)
+		[
+			SNew(STextBlock)
+			.ColorAndOpacity(FStyleColors::Foreground)
+			.Text(LOCTEXT("Loaded", "Loaded"))
+		];
+
+	Panel->AddSlot(1, Row)
+		[
+			SNew(STextBlock)
+			.Margin(FMargin(5, 0))
+			.Text(FText::FromString(SingleDecimalFormat(DDCResourceStatsTotal.LoadSizeMB) + TEXT(" MB")))
+			.Justification(ETextJustify::Right)
+		];
+
+	Row++;
+	
+	Panel->AddSlot(0, Row)
+		[
+			SNew(STextBlock)
+			.ColorAndOpacity(FStyleColors::Foreground)
+			.Text(LOCTEXT("Built", "Built"))
+		];
+
+	Panel->AddSlot(1, Row)
+		[
+			SNew(STextBlock)
+			.Margin(FMargin(5, 0))
+			.Text(FText::FromString(SingleDecimalFormat(DDCResourceStatsTotal.BuildSizeMB) + TEXT(" MB")))
+			.Justification(ETextJustify::Right)
+		];
+
+	Row++;*/
+
+	Panel->AddSlot(0, Row)
+		[
+		SNew(STextBlock)
+		.Font(FCoreStyle::GetDefaultFontStyle("Bold", 10))
+		.ColorAndOpacity(FStyleColors::Foreground)
+		.Text(LOCTEXT("Remote Storage", "Remote Storage"))
+
+		];
+
+	Row++;
+
+	Panel->AddSlot(0, Row)
+		[
+			SNew(STextBlock)
+			.Text(LOCTEXT("Connected", "Connected"))
+		];
+
+	Panel->AddSlot(1, Row)
+		[
+			SNew(STextBlock)
+			.Text(FText::FromString(FDerivedDataInformation::GetHasRemoteCache()? TEXT("True") : TEXT("False")))
+			.Justification(ETextJustify::Right)
+		];
+
+	Row++;
+
+	Panel->AddSlot(0, Row)
+		[
+			SNew(STextBlock)
+			.Text(LOCTEXT("Downloaded", "Downloaded"))
+		];
+
+	Panel->AddSlot(1, Row)
+		[
+			SNew(STextBlock)
+			.Text(FText::FromString(SingleDecimalFormat(DownloadedBytesMB) + TEXT(" MB")))
+			.Justification(ETextJustify::Right)
+		];
+
+	Row++;
+
+	Panel->AddSlot(0, Row)
+		[
+			SNew(STextBlock)
+			.Text(LOCTEXT("Uploaded", "Uploaded"))
+		];
+
+	Panel->AddSlot(1, Row)
+		[
+			SNew(STextBlock)
+			.Text(FText::FromString(SingleDecimalFormat(UploadedBytesMB) + TEXT(" MB")))
+			.Justification(ETextJustify::Right)
+		];
+
+	Row++;
+
+	return Panel;
+}
+
+
+
+void SDerivedDataResourceUsageDialog::Construct(const FArguments& InArgs)
+{
+	this->ChildSlot
+	[
+		SNew(SVerticalBox)
+
+		+ SVerticalBox::Slot()
+			.AutoHeight()
+			.Padding(0, 20, 0, 0)
+			.Expose(GridSlot)
+			[
+				GetGridPanel()
+			]
+	];
+
+	RegisterActiveTimer(0.5f, FWidgetActiveTimerDelegate::CreateSP(this, &SDerivedDataResourceUsageDialog::UpdateGridPanels));
+}
+
+EActiveTimerReturnType SDerivedDataResourceUsageDialog::UpdateGridPanels(double InCurrentTime, float InDeltaTime)
+{
+	(*GridSlot)
+	[
+		GetGridPanel()
+	];
+
+	SlatePrepass(GetPrepassLayoutScaleMultiplier());
+
+	return EActiveTimerReturnType::Continue;
+}
+
+TSharedRef<SWidget> SDerivedDataResourceUsageDialog::GetGridPanel()
+{	
+	
+	TArray<FDerivedDataCacheResourceStat> DDCResourceStats;
+
+	// Grab the resource stats
+	GetDerivedDataCacheRef().GatherResourceStats(DDCResourceStats);
+
+	// Sort results on asscending Load size
+	DDCResourceStats.Sort([](const FDerivedDataCacheResourceStat& LHS, const FDerivedDataCacheResourceStat& RHS) { return LHS.LoadSizeMB > RHS.LoadSizeMB; });
+
+	FDerivedDataCacheResourceStat DDCResourceStatsTotal(TEXT("Total"));
+
+	// Accumulate Totals
+	for (const FDerivedDataCacheResourceStat& Stat : DDCResourceStats)
+	{
+		DDCResourceStatsTotal.Accumulate(Stat);
+	}
+
+	PRAGMA_DISABLE_DEPRECATION_WARNINGS
+	TSharedRef<FDerivedDataCacheStatsNode> RootUsage = GetDerivedDataCache()->GatherUsageStats();
+	PRAGMA_ENABLE_DEPRECATION_WARNINGS
+		TArray<TSharedRef<const FDerivedDataCacheStatsNode>> LeafUsageStats;
+	RootUsage->ForEachDescendant([&LeafUsageStats](TSharedRef<const FDerivedDataCacheStatsNode> Node) {
+		if (Node->Children.Num() == 0)
+		{
+			LeafUsageStats.Add(Node);
+		}
+		});
+
+	TSharedRef<SGridPanel> Panel =
+		SNew(SGridPanel);
+
+	int32 Row = 0;
+
+	const float RowMargin = 0.0f;
+	const float TitleMargin = 10.0f;
+	const float ColumnMargin = 10.0f;
+
+	Panel->AddSlot(2, Row)
+		[
+		SNew(STextBlock)
+		.Margin(FMargin(ColumnMargin, RowMargin))
+		.ColorAndOpacity(FStyleColors::Foreground)
+		.Font(FCoreStyle::GetDefaultFontStyle("Bold", 10))
+		.Text(LOCTEXT("Loaded", "Loaded"))
+		.Justification(ETextJustify::Center)
+		];
+
+
+	Panel->AddSlot(5, Row)
+		[
+		SNew(STextBlock)
+		.Margin(FMargin(ColumnMargin, RowMargin))
+		.ColorAndOpacity(FStyleColors::Foreground)
+		.Font(FCoreStyle::GetDefaultFontStyle("Bold", 10))
+		.Text(LOCTEXT("Built", "Built"))
+		.Justification(ETextJustify::Center)
+		];
+
+	Row++;
+
+	Panel->AddSlot(0, Row)
+		[
+		SNew(STextBlock)
+		.ColorAndOpacity(FStyleColors::Foreground)
+		.Margin(FMargin(ColumnMargin, RowMargin, 0.0f, TitleMargin))
+		.Font(FCoreStyle::GetDefaultFontStyle("Bold", 10))
+		.Text(LOCTEXT("Asset", "Asset"))
+		];
+
+	Panel->AddSlot(1, Row)
+		[
+	
+		SNew(STextBlock)
+		.ColorAndOpacity(FStyleColors::Foreground)
+		.Margin(FMargin(ColumnMargin, RowMargin, 0.0f, TitleMargin))
+		.Font(FCoreStyle::GetDefaultFontStyle("Bold", 10))
+		.Text(LOCTEXT("Count", "Count"))
+		.Justification(ETextJustify::Center)
+		];
+
+	Panel->AddSlot(2, Row)
+		[
+		SNew(STextBlock)
+		.ColorAndOpacity(FStyleColors::Foreground)
+		.Margin(FMargin(ColumnMargin, RowMargin, 0.0f, TitleMargin))
+		.Font(FCoreStyle::GetDefaultFontStyle("Bold", 10))
+		.Text(LOCTEXT("Time (Sec)", "Time (Sec)"))
+		.Justification(ETextJustify::Center)
+		];
+
+	Panel->AddSlot(3, Row)
+		[
+		SNew(STextBlock)
+		.ColorAndOpacity(FStyleColors::Foreground)
+		.Margin(FMargin(ColumnMargin, RowMargin, 0.0f, TitleMargin))
+		.Font(FCoreStyle::GetDefaultFontStyle("Bold", 10))
+		.Text(LOCTEXT("Size (MB)", "Size (MB)"))
+		.Justification(ETextJustify::Center)
+		];
+
+	Panel->AddSlot(4, Row)
+		[
+		SNew(STextBlock)
+		.ColorAndOpacity(FStyleColors::Foreground)
+		.Margin(FMargin(ColumnMargin, RowMargin, 0.0f, TitleMargin))
+		.Font(FCoreStyle::GetDefaultFontStyle("Bold", 10))
+		.Text(LOCTEXT("Count", "Count"))
+		.Justification(ETextJustify::Center)
+		];
+
+	Panel->AddSlot(5, Row)
+		[
+		SNew(STextBlock)
+		.ColorAndOpacity(FStyleColors::Foreground)
+		.Margin(FMargin(ColumnMargin, RowMargin, 0.0f, TitleMargin))
+		.Font(FCoreStyle::GetDefaultFontStyle("Bold", 10))
+		.Text(LOCTEXT("Time (Sec)", "Time (Sec)"))
+		.Justification(ETextJustify::Center)
+		];
+
+	Panel->AddSlot(6, Row)
+		[
+		SNew(STextBlock)
+		.Margin(FMargin(ColumnMargin, RowMargin, 0.0f, TitleMargin))
+		.ColorAndOpacity(FStyleColors::Foreground)
+		.Font(FCoreStyle::GetDefaultFontStyle("Bold", 10))
+		.Text(LOCTEXT("Size (MB)", "Size (MB)"))
+		.Justification(ETextJustify::Center)
+		];
+
+	Row++;
+
+
+	for (const FDerivedDataCacheResourceStat& Stat : DDCResourceStats)
+	{	
+		Panel->AddSlot(0, Row)
+		[
+			SNew(STextBlock)
+			.Margin(FMargin(ColumnMargin, RowMargin))
+			.Text(FText::FromString(Stat.AssetType))
+			.Justification(ETextJustify::Left)
+		];
+
+		Panel->AddSlot(1, Row)
+		[
+			SNew(STextBlock)
+			.Margin(FMargin(ColumnMargin, RowMargin))
+			.Text(FText::FromString(FString::Printf(TEXT("%u"),Stat.LoadCount)))
+			.Justification(ETextJustify::Center)
+		];
+
+		Panel->AddSlot(2, Row)
+		[
+			SNew(STextBlock)
+			.Margin(FMargin(ColumnMargin, RowMargin))
+			.Text(FText::FromString(SingleDecimalFormat(Stat.LoadTimeSec)))
+			.Justification(ETextJustify::Center)
+		];
+
+		Panel->AddSlot(3, Row)
+		[
+			SNew(STextBlock)
+			.Margin(FMargin(ColumnMargin, RowMargin))
+			.Text(FText::FromString(SingleDecimalFormat(Stat.LoadSizeMB)))
+			.Justification(ETextJustify::Center)
+		];
+
+		Panel->AddSlot(4, Row)
+		[
+			SNew(STextBlock)
+			.Margin(FMargin(ColumnMargin, RowMargin))
+			.Text(FText::FromString(FString::Printf(TEXT("%u"),Stat.BuildCount)))
+			.Justification(ETextJustify::Center)
+		];
+
+		Panel->AddSlot(5, Row)
+		[
+			SNew(STextBlock)
+			.Margin(FMargin(ColumnMargin, RowMargin))
+			.Text(FText::FromString(SingleDecimalFormat(Stat.BuildTimeSec)))
+			.Justification(ETextJustify::Center)
+		];
+
+		Panel->AddSlot(6, Row)
+		[
+			SNew(STextBlock)
+			.Margin(FMargin(ColumnMargin, RowMargin))
+			.Text(FText::FromString(SingleDecimalFormat(Stat.BuildSizeMB)))
+			.Justification(ETextJustify::Center)
+		];
+
+		Row++;
+	}
+	
+	Panel->AddSlot(0, Row)
+		[
+			SNew(STextBlock)
+			.Margin(FMargin(ColumnMargin, RowMargin))
+			.Text(FText::FromString(DDCResourceStatsTotal.AssetType))
+			.ColorAndOpacity(FStyleColors::AccentWhite)
+			.Justification(ETextJustify::Left)
+		];
+
+	Panel->AddSlot(1, Row)
+		[
+			SNew(STextBlock)
+			.Margin(FMargin(ColumnMargin, RowMargin))
+			.Text(FText::FromString(FString::Printf(TEXT("%u"), DDCResourceStatsTotal.LoadCount)))
+			.ColorAndOpacity(FStyleColors::Foreground)
+			.Justification(ETextJustify::Center)
+		];
+
+	Panel->AddSlot(2, Row)
+		[
+			SNew(STextBlock)
+			.Margin(FMargin(ColumnMargin, RowMargin))
+			.Text(FText::FromString(SingleDecimalFormat(DDCResourceStatsTotal.LoadTimeSec)))
+			.ColorAndOpacity(FStyleColors::Foreground)
+			.Justification(ETextJustify::Center)
+		];
+
+	Panel->AddSlot(3, Row)
+		[
+			SNew(STextBlock)
+			.Margin(FMargin(ColumnMargin, RowMargin))
+			.Text(FText::FromString(SingleDecimalFormat(DDCResourceStatsTotal.LoadSizeMB)))
+			.ColorAndOpacity(FStyleColors::Foreground)
+			.Justification(ETextJustify::Center)
+		];
+
+	Panel->AddSlot(4, Row)
+		[
+			SNew(STextBlock)
+			.Margin(FMargin(ColumnMargin, RowMargin))
+			.Text(FText::FromString(FString::Printf(TEXT("%u"), DDCResourceStatsTotal.BuildCount)))
+			.ColorAndOpacity(FStyleColors::Foreground)
+			.Justification(ETextJustify::Center)
+		];
+
+	Panel->AddSlot(5, Row)
+		[
+			SNew(STextBlock)
+			.Margin(FMargin(ColumnMargin, RowMargin))
+			.Text(FText::FromString(SingleDecimalFormat(DDCResourceStatsTotal.BuildTimeSec)))
+			.ColorAndOpacity(FStyleColors::Foreground)
+			.Justification(ETextJustify::Center)
+		];
+
+	Panel->AddSlot(6, Row)
+		[
+			SNew(STextBlock)
+			.Margin(FMargin(ColumnMargin, RowMargin))
+			.Text(FText::FromString(SingleDecimalFormat(DDCResourceStatsTotal.BuildSizeMB)))
+			.ColorAndOpacity(FStyleColors::Foreground)
+			.Justification(ETextJustify::Center)
+		];
+
+	return Panel;
+}
+
+
+void SDerivedDataCacheStatisticsDialog::Construct(const FArguments& InArgs)
+{
+	this->ChildSlot
+	[
+		SNew(SVerticalBox)
+
+		+ SVerticalBox::Slot()
+		.AutoHeight()
+		.Padding(0, 20, 0, 0)
+		.Expose(GridSlot)
+		[
+			GetGridPanel()
+		]
+
+	];
+
+	RegisterActiveTimer(0.5f, FWidgetActiveTimerDelegate::CreateSP(this, &SDerivedDataCacheStatisticsDialog::UpdateGridPanels));
+}
+
+EActiveTimerReturnType SDerivedDataCacheStatisticsDialog::UpdateGridPanels(double InCurrentTime, float InDeltaTime)
+{
+	(*GridSlot)
+	[
+		GetGridPanel()
+	];
+
+	SlatePrepass(GetPrepassLayoutScaleMultiplier());
+
+	return EActiveTimerReturnType::Continue;
+} 
+
+
+
+TSharedRef<SWidget> SDerivedDataCacheStatisticsDialog::GetGridPanel()
+{
+	PRAGMA_DISABLE_DEPRECATION_WARNINGS
+	TSharedRef<FDerivedDataCacheStatsNode> RootUsage = GetDerivedDataCache()->GatherUsageStats();
+	PRAGMA_ENABLE_DEPRECATION_WARNINGS
+	TArray<TSharedRef<const FDerivedDataCacheStatsNode>> LeafUsageStats;
+	RootUsage->ForEachDescendant([&LeafUsageStats](TSharedRef<const FDerivedDataCacheStatsNode> Node) {
+		if (Node->Children.Num() == 0)
+		{
+			LeafUsageStats.Add(Node);
+		}
+	});
+
+	TSharedRef<SGridPanel> Panel =
+		SNew(SGridPanel);
+
+	int32 Row = 0;
+
+	const float RowMargin = 0.0f;
+	const float TitleMargin = 10.0f;
+	const float ColumnMargin = 10.0f;
+
+	Panel->AddSlot(0, Row)
+	[
+		SNew(STextBlock)
+		.ColorAndOpacity(FStyleColors::Foreground)
+		.Margin(FMargin(ColumnMargin, RowMargin, 0.0f, TitleMargin))
+		.Font(FCoreStyle::GetDefaultFontStyle("Bold", 10))
+		.Text(LOCTEXT("Cache", "Cache"))
+	];
+
+	Panel->AddSlot(1, Row)
+	[
+		SNew(STextBlock)
+		.Margin(FMargin(ColumnMargin, RowMargin, 0.0f, TitleMargin))
+		.ColorAndOpacity(FStyleColors::Foreground)
+		.Font(FCoreStyle::GetDefaultFontStyle("Bold", 10))
+		.Justification(ETextJustify::Right)
+		.Text(LOCTEXT("Speed", "Speed"))
+	];
+
+	Panel->AddSlot(2, Row)
+	[
+		SNew(STextBlock)
+		.Margin(FMargin(ColumnMargin, RowMargin, 0.0f, TitleMargin))
+		.ColorAndOpacity(FStyleColors::Foreground)
+		.Font(FCoreStyle::GetDefaultFontStyle("Bold", 10))
+		.Justification(ETextJustify::Right)
+		.Text(LOCTEXT("HitPercentage", "Hit%"))
+	];
+
+	Panel->AddSlot(3, Row)
+	[
+		SNew(STextBlock)
+		.Margin(FMargin(ColumnMargin, RowMargin, 0.0f, TitleMargin))
+		.ColorAndOpacity(FStyleColors::Foreground)
+		.Font(FCoreStyle::GetDefaultFontStyle("Bold", 10))
+		.Justification(ETextJustify::Right)
+		.Text(LOCTEXT("Read", "Read"))
+	];
+
+	Panel->AddSlot(4, Row)
+		[
+		SNew(STextBlock)
+		.Margin(FMargin(ColumnMargin, RowMargin, 0.0f, TitleMargin))
+		.ColorAndOpacity(FStyleColors::Foreground)
+		.Font(FCoreStyle::GetDefaultFontStyle("Bold", 10))
+		.Justification(ETextJustify::Right)
+		.Text(LOCTEXT("Write", "Write"))
+		];
+
+	Panel->AddSlot(5, Row)
+	[
+		SNew(SImage)
+		.Image(FAppStyle::Get().GetBrush("Icons.Edit"))
+		.ColorAndOpacity(FStyleColors::Foreground)
+		
+	];
+
+	Row++;
+
+	double SumTotalGetMB =0.0;
+	double SumTotalPutMB = 0.0;
+
+	for (TSharedRef<const FDerivedDataCacheStatsNode> Node : LeafUsageStats)
+	{
+		FDerivedDataCacheUsageStats Stats;
+
+		const FDerivedDataBackendInterface* Backend = Node->GetBackendInterface();
+
+		TSharedRef<FDerivedDataCacheStatsNode> Usage = Backend->GatherUsageStats();
+		for (const auto& KVP : Usage->Stats)
+		{
+			Stats.Combine(KVP.Value);
+		}
+
+		const int64 TotalGetBytes = Stats.GetStats.GetAccumulatedValueAnyThread(FCookStats::CallStats::EHitOrMiss::Hit, FCookStats::CallStats::EStatType::Bytes);
+		const int64 TotalPutBytes = Stats.PutStats.GetAccumulatedValueAnyThread(FCookStats::CallStats::EHitOrMiss::Hit, FCookStats::CallStats::EStatType::Bytes);
+		const int64 TotalPrefetchBytes = Stats.PrefetchStats.GetAccumulatedValueAnyThread(FCookStats::CallStats::EHitOrMiss::Hit, FCookStats::CallStats::EStatType::Bytes);
+
+		const int64 TotalGets_Hit = Stats.GetStats.GetAccumulatedValueAnyThread(FCookStats::CallStats::EHitOrMiss::Hit, FCookStats::CallStats::EStatType::Counter);
+		const int64 TotalGets_Miss = Stats.GetStats.GetAccumulatedValueAnyThread(FCookStats::CallStats::EHitOrMiss::Miss, FCookStats::CallStats::EStatType::Counter);
+
+		const int64 TotalRequests = TotalGets_Hit + TotalGets_Miss;
+
+		//if (TotalGetBytes > 0 || TotalPutBytes > 0)
+		{
+			Panel->AddSlot(0, Row)
+			[
+				SNew(STextBlock)
+				.Margin(FMargin(ColumnMargin, RowMargin))
+				.Text(FText::FromString(Backend->GetDisplayName()))
+			];
+
+			Panel->AddSlot(1, Row)
+			[
+				SNew(STextBlock)
+				.Margin(FMargin(ColumnMargin, RowMargin))
+				.Justification(ETextJustify::Right)
+				.Text(FText::FromString(LexToString(Backend->GetSpeedClass())))
+			];
+
+			const double HitPercentage = TotalRequests > 0 ? (100.0 * (TotalGets_Hit / (double)TotalRequests)) : 0.0;
+
+			Panel->AddSlot(2, Row)
+				.HAlign(HAlign_Right)
+				[
+				SNew(STextBlock)
+				.Margin(FMargin(ColumnMargin, RowMargin))
+				.Justification(ETextJustify::Right)
+				.Text(FText::FromString(SingleDecimalFormat(HitPercentage) + TEXT(" %")))
+				];
+
+			const double TotalGetMB = FUnitConversion::Convert(TotalGetBytes, EUnit::Bytes, EUnit::Megabytes);
+
+			SumTotalGetMB += TotalGetMB;
+
+			Panel->AddSlot(3, Row)
+			.HAlign(HAlign_Right)
+			[
+				SNew(STextBlock)
+				.Margin(FMargin(ColumnMargin, RowMargin))
+				.Justification(ETextJustify::Right)
+				.Text(FText::FromString(SingleDecimalFormat(TotalGetMB) + TEXT(" MB")))
+			];
+
+			const double TotalPutMB = FUnitConversion::Convert(TotalPutBytes, EUnit::Bytes, EUnit::Megabytes);
+
+			SumTotalPutMB += TotalPutMB;
+
+			Panel->AddSlot(4, Row)
+			.HAlign(HAlign_Right)
+			[
+				SNew(STextBlock)
+				.Margin(FMargin(ColumnMargin, RowMargin))
+				.Justification(ETextJustify::Right)
+				.Text(FText::FromString(SingleDecimalFormat(TotalPutMB) + TEXT(" MB")))
+			];
+
+			Panel->AddSlot(5, Row)
+			.HAlign(HAlign_Center)
+			[
+				SNew(SImage)
+				.Image(FAppStyle::Get().GetBrush("Icons.Check"))
+				.Visibility(Backend->IsWritable() ? EVisibility::Visible : EVisibility::Hidden)
+			];
+
+			Row++;
+		}
+	}
+
+	Panel->AddSlot(0, Row)
+		[
+			SNew(STextBlock)
+			.Text(FText::FromString(TEXT("Total")))
+			.Margin(FMargin(ColumnMargin, RowMargin))
+			.ColorAndOpacity(FStyleColors::AccentWhite)
+			.Justification(ETextJustify::Left)	
+		];
+
+	Panel->AddSlot(3, Row)
+		.HAlign(HAlign_Right)
+		[
+		SNew(STextBlock)
+		.Margin(FMargin(ColumnMargin, RowMargin))
+		.Justification(ETextJustify::Right)
+		.ColorAndOpacity(FStyleColors::Foreground)
+		.Text(FText::FromString(SingleDecimalFormat(SumTotalGetMB) + TEXT(" MB")))
+		];
+
+	Panel->AddSlot(4, Row)
+		.HAlign(HAlign_Right)
+		[
+		SNew(STextBlock)
+		.Margin(FMargin(ColumnMargin, RowMargin))
+		.Justification(ETextJustify::Right)
+		.ColorAndOpacity(FStyleColors::Foreground)
+		.Text(FText::FromString(SingleDecimalFormat(SumTotalPutMB) + TEXT(" MB")))
+		];
+
+	return Panel;
+}
+
+#undef LOCTEXT_NAMESPACE
\ No newline at end of file
diff --git a/Engine/Source/Editor/DerivedDataEditor/Private/SDerivedDataDialogs.h b/Engine/Source/Editor/DerivedDataEditor/Private/SDerivedDataDialogs.h
new file mode 100644
--- /dev/null
+++ b/Engine/Source/Editor/DerivedDataEditor/Private/SDerivedDataDialogs.h
@@ -0,0 +1,57 @@
+// Copyright Epic Games, Inc. All Rights Reserved.
+
+#pragma once
+
+#include "CoreMinimal.h"
+#include "Widgets/DeclarativeSyntaxSupport.h"
+#include "Widgets/SWidget.h"
+#include "Widgets/SCompoundWidget.h"
+#include "Widgets/SBoxPanel.h"
+
+class SDerivedDataRemoteStoreDialog : public SCompoundWidget
+{
+	SLATE_BEGIN_ARGS(SDerivedDataRemoteStoreDialog) {}
+	SLATE_END_ARGS()
+
+	void Construct(const FArguments& InArgs);
+
+private:
+
+	TSharedRef<SWidget> GetGridPanel();
+
+	EActiveTimerReturnType UpdateGridPanels(double InCurrentTime, float InDeltaTime);
+
+	SVerticalBox::FSlot* GridSlot = nullptr;
+};
+
+class SDerivedDataCacheStatisticsDialog : public SCompoundWidget
+{
+	SLATE_BEGIN_ARGS(SDerivedDataCacheStatisticsDialog) {}
+	SLATE_END_ARGS()
+
+	void Construct(const FArguments& InArgs);
+
+private:
+
+	TSharedRef<SWidget> GetGridPanel();
+	
+	EActiveTimerReturnType UpdateGridPanels(double InCurrentTime, float InDeltaTime);
+
+	SVerticalBox::FSlot* GridSlot = nullptr;
+};
+
+class SDerivedDataResourceUsageDialog : public SCompoundWidget
+{
+	SLATE_BEGIN_ARGS(SDerivedDataResourceUsageDialog) {}
+	SLATE_END_ARGS()
+
+		void Construct(const FArguments& InArgs);
+
+private:
+
+	TSharedRef<SWidget> GetGridPanel();
+	
+	EActiveTimerReturnType UpdateGridPanels(double InCurrentTime, float InDeltaTime);
+
+	SVerticalBox::FSlot* GridSlot = nullptr;
+};
diff --git a/Engine/Source/Editor/DerivedDataEditor/Private/SDerivedDataStatusBar.cpp b/Engine/Source/Editor/DerivedDataEditor/Private/SDerivedDataStatusBar.cpp
new file mode 100644
--- /dev/null
+++ b/Engine/Source/Editor/DerivedDataEditor/Private/SDerivedDataStatusBar.cpp
@@ -0,0 +1,340 @@
+// Copyright Epic Games, Inc. All Rights Reserved.
+
+#include "SDerivedDataStatusBar.h"
+#include "DerivedDataEditorModule.h"
+#include "DerivedDataInformation.h"
+#include "SDerivedDataDialogs.h"
+#include "SDerivedDataCacheSettings.h"
+#include "DerivedDataCacheUsageStats.h"
+#include "DerivedDataCacheInterface.h"
+#include "DerivedDataBackendInterface.h"
+#include "Widgets/SBoxPanel.h"
+#include "Widgets/SToolTip.h"
+#include "Widgets/Text/STextBlock.h"
+#include "Widgets/Layout/SBorder.h"
+#include "EditorStyleSet.h"
+#include "EditorFontGlyphs.h"
+#include "DerivedDataCacheUsageStats.h"
+#include "Stats/Stats.h"
+#include "Widgets/SOverlay.h"
+#include "Settings/EditorSettings.h"
+#include "Framework/Notifications/NotificationManager.h"
+#include "Widgets/Notifications/SNotificationList.h"
+#include "Async/Future.h"
+#include "Settings/EditorProjectSettings.h"
+#include "Modules/ModuleManager.h"
+#include "ISettingsModule.h"
+#include "ToolMenuContext.h"
+#include "ToolMenus.h"
+#include "Widgets/Input/SComboButton.h"
+#include "Widgets/Images/SImage.h"
+#include "Styling/StyleColors.h"
+
+#define LOCTEXT_NAMESPACE "DerivedDataEditor"
+
+
+TSharedRef<FUICommandList> FDerivedDataStatusBarMenuCommands::ActionList(new FUICommandList());
+
+FDerivedDataStatusBarMenuCommands::FDerivedDataStatusBarMenuCommands()
+	: TCommands<FDerivedDataStatusBarMenuCommands>
+	(
+		"DerivedDataSettings",
+		NSLOCTEXT("Contexts", "Derived Data", "Derived Data"),
+		"LevelEditor",
+		FEditorStyle::GetStyleSetName()
+		)
+{}
+
+void FDerivedDataStatusBarMenuCommands::RegisterCommands()
+{
+	UI_COMMAND(ChangeSettings, "Change Cache Settings", "Opens a dialog to change Cache settings.", EUserInterfaceActionType::Button, FInputChord());
+	UI_COMMAND(ViewCacheStatistics, "View Cache Statistics", "Opens the Cache Statistics panel.", EUserInterfaceActionType::Button, FInputChord());
+	UI_COMMAND(ViewResourceUsage, "View Resource Usage", "Opens the Resource Usage panel.", EUserInterfaceActionType::Button, FInputChord());
+
+	ActionList->MapAction(
+		ChangeSettings,
+		FExecuteAction::CreateStatic(&FDerivedDataStatusBarMenuCommands::ChangeSettings_Clicked)
+	);
+
+	ActionList->MapAction(
+		ViewCacheStatistics,
+		FExecuteAction::CreateStatic(&FDerivedDataStatusBarMenuCommands::ViewCacheStatistics_Clicked)
+	);
+
+	ActionList->MapAction(
+		ViewResourceUsage,
+		FExecuteAction::CreateStatic(&FDerivedDataStatusBarMenuCommands::ViewResourceUsage_Clicked)
+	);
+}
+
+void FDerivedDataStatusBarMenuCommands::ChangeSettings_Clicked()
+{
+	FDerivedDataEditorModule& DerivedDataEditorModule = FModuleManager::LoadModuleChecked<FDerivedDataEditorModule>("DerivedDataEditor");
+	DerivedDataEditorModule.ShowSettingsDialog();
+}
+
+void FDerivedDataStatusBarMenuCommands::ViewCacheStatistics_Clicked()
+{
+	FDerivedDataEditorModule& DerivedDataEditorModule = FModuleManager::LoadModuleChecked<FDerivedDataEditorModule>("DerivedDataEditor");
+	DerivedDataEditorModule.ShowCacheStatisticsTab();
+}
+
+void FDerivedDataStatusBarMenuCommands::ViewResourceUsage_Clicked()
+{
+	FDerivedDataEditorModule& DerivedDataEditorModule = FModuleManager::LoadModuleChecked<FDerivedDataEditorModule>("DerivedDataEditor");
+	DerivedDataEditorModule.ShowResourceUsageTab();
+}
+
+TSharedRef<SWidget> SDerivedDataStatusBarWidget::CreateStatusBarMenu()
+{
+	UToolMenu* Menu = UToolMenus::Get()->RegisterMenu("StatusBar.ToolBar.DDC", NAME_None, EMultiBoxType::Menu, false);
+
+	{
+		FToolMenuSection& Section = Menu->AddSection("DDCMenuSettingsSection", LOCTEXT("DDCMenuSettingsSection", "Settings"));
+
+		Section.AddMenuEntry(
+			FDerivedDataStatusBarMenuCommands::Get().ChangeSettings,
+			TAttribute<FText>(),
+			TAttribute<FText>(),
+			FSlateIcon(FEditorStyle::GetStyleSetName(), "DerivedData.CacheSettings")
+		);
+	}
+
+	{
+		FToolMenuSection& Section = Menu->AddSection("DDCMenuStatisticsSection", LOCTEXT("DDCMenuStatisticsSection", "Statistics"));
+
+		Section.AddMenuEntry(
+			FDerivedDataStatusBarMenuCommands::Get().ViewCacheStatistics,
+			TAttribute<FText>(),
+			TAttribute<FText>(),
+			FSlateIcon(FEditorStyle::GetStyleSetName(), "DerivedData.CacheStatistics")
+		);
+
+		Section.AddMenuEntry(
+			FDerivedDataStatusBarMenuCommands::Get().ViewResourceUsage,
+			TAttribute<FText>(),
+			TAttribute<FText>(),
+			FSlateIcon(FEditorStyle::GetStyleSetName(), "DerivedData.ResourceUsage")
+		);
+	}
+
+	return UToolMenus::Get()->GenerateWidget("StatusBar.ToolBar.DDC", FToolMenuContext(FDerivedDataStatusBarMenuCommands::ActionList));
+}
+
+void SDerivedDataStatusBarWidget::Construct(const FArguments& InArgs)
+{	
+	/*FDerivedDataCacheInterface::FOnDDCNotification& DDCNotificationEvent = GetDerivedDataCacheRef().GetDDCNotificationEvent();
+
+	if (bSubscribe)
+	{
+		DDCNotificationEvent.AddRaw(this, &FDDCNotifications::OnDDCNotificationEvent);
+	}
+	else
+	{
+		DDCNotificationEvent.RemoveAll(this);
+	}*/
+
+	BusyPulseSequence = FCurveSequence(0.f, 1.0f, ECurveEaseFunction::QuadInOut);
+	FadeGetSequence = FCurveSequence(0.f, 0.5f, ECurveEaseFunction::Linear);
+	FadePutSequence = FCurveSequence(0.f, 0.5f, ECurveEaseFunction::Linear);
+
+	this->ChildSlot
+	[
+		SNew(SComboButton)
+		.ContentPadding(FMargin(6.0f, 0.0f))
+		.MenuPlacement(MenuPlacement_AboveAnchor)
+		.ComboButtonStyle(&FAppStyle::Get().GetWidgetStyle<FComboButtonStyle>("StatusBar.StatusBarComboButton"))
+		.ButtonContent()
+		[
+
+			SNew(SHorizontalBox)
+	
+			+ SHorizontalBox::Slot()
+			.AutoWidth()
+			.VAlign(VAlign_Center)
+			.Padding(0, 0, 3, 0)
+			[
+				SNew(SOverlay)
+				+ SOverlay::Slot()
+				.HAlign(HAlign_Center)
+				.VAlign(VAlign_Top)
+				[
+					SNew(SImage)
+					.Image_Lambda([this] {  return FDerivedDataInformation::GetHasRemoteCache() ? FAppStyle::Get().GetBrush("DerivedData.RemoteStorage.Connected") : FAppStyle::Get().GetBrush("DerivedData.RemoteStorage.Unavailable"); })			
+					.ToolTipText_Lambda([this] { return GetRemoteCacheToolTipText(); })
+				]
+
+				+ SOverlay::Slot()
+				.HAlign(HAlign_Center)
+				.VAlign(VAlign_Top)
+				.Padding(0, 0, 4, 4)
+				[
+					SNew(SImage)
+					.Image(FAppStyle::Get().GetBrush("DerivedData.Uploading"))
+					.ColorAndOpacity_Lambda([this] { return bPutActive ? FLinearColor::White.CopyWithNewOpacity(0.5f + (0.5f * FMath::MakePulsatingValue(FadePutSequence.GetLerp(), 1))) : FLinearColor(0,0,0,0); })
+					.ToolTipText_Lambda([this] { return GetRemoteCacheToolTipText(); })
+				]
+
+				+ SOverlay::Slot()
+				.HAlign(HAlign_Center)
+				.VAlign(VAlign_Bottom)
+				.Padding(4, 4, 0, 0)
+				[
+					SNew(SImage)
+					.Image(FAppStyle::Get().GetBrush("DerivedData.Downloading"))
+					.ColorAndOpacity_Lambda([this] { return bGetActive ? FLinearColor::White.CopyWithNewOpacity(0.5f + (0.5f * FMath::MakePulsatingValue(FadeGetSequence.GetLerp(), 1))) : FLinearColor(0, 0, 0, 0); })
+					.ToolTipText_Lambda([this] { return GetRemoteCacheToolTipText(); })
+				]
+			]
+
+			+ SHorizontalBox::Slot()
+			.AutoWidth()
+			.VAlign(VAlign_Center)
+			.Padding(0, 0, 10, 0)
+			[
+				SNew(STextBlock)
+				.Text_Lambda([this] { return GetTitleText(); })
+				.ToolTipText_Lambda([this] { return GetTitleToolTipText(); })
+			]
+		]
+		.OnGetMenuContent(FOnGetContent::CreateRaw(this, &SDerivedDataStatusBarWidget::CreateStatusBarMenu))
+	];
+
+	LastDDCGetTime = FDerivedDataInformation::GetCacheActivityTimeSeconds(true, false);
+	LastDDCPutTime = FDerivedDataInformation::GetCacheActivityTimeSeconds(false, false);
+	RegisterActiveTimer(0.5f, FWidgetActiveTimerDelegate::CreateSP(this, &SDerivedDataStatusBarWidget::UpdateBusyIndicator));
+	RegisterActiveTimer(5.0f, FWidgetActiveTimerDelegate::CreateSP(this, &SDerivedDataStatusBarWidget::UpdateWarnings));
+}
+
+EActiveTimerReturnType SDerivedDataStatusBarWidget::UpdateBusyIndicator(double InCurrentTime, float InDeltaTime)
+{
+	const double OldLastDDCGetTime = LastDDCGetTime;
+	const double OldLastDDCPutTime = LastDDCPutTime;
+
+	LastDDCGetTime = FDerivedDataInformation::GetCacheActivityTimeSeconds(true, false);
+	LastDDCPutTime = FDerivedDataInformation::GetCacheActivityTimeSeconds(false, false);
+
+	bGetActive = OldLastDDCGetTime != LastDDCGetTime;
+	bPutActive = OldLastDDCPutTime != LastDDCPutTime;
+	bBusy = GetDerivedDataCache()->AnyAsyncRequestsRemaining();
+
+	FadeGetSequence.PlayRelative(this->AsShared(), bGetActive);
+	FadePutSequence.PlayRelative(this->AsShared(), bPutActive);
+
+	if (bBusy)
+	{
+		if (!BusyPulseSequence.IsPlaying())
+		{
+			BusyPulseSequence.Play(this->AsShared(), true);
+		}
+	}
+	else
+	{
+		BusyPulseSequence.JumpToEnd();
+		BusyPulseSequence.Pause();
+	}
+
+	return EActiveTimerReturnType::Continue;
+}
+
+EActiveTimerReturnType SDerivedDataStatusBarWidget::UpdateWarnings(double InCurrentTime, float InDeltaTime)
+{
+	const UEditorSettings* Settings = GetDefault<UEditorSettings>();
+	const UDDCProjectSettings* DDCProjectSettings = GetDefault<UDDCProjectSettings>();
+
+	if (DDCProjectSettings->RecommendEveryoneSetupAGlobalLocalDDCPath && Settings->GlobalLocalDDCPath.Path.IsEmpty())
+	{
+		TPromise<TWeakPtr<SNotificationItem>> NotificationPromise;
+
+		FNotificationInfo Info(LOCTEXT("SharedProjectLocalDDC", "This project recommends you setup the 'Global Local DDC Path', \nso that all copies of this project use the same local DDC cache."));
+		Info.bUseSuccessFailIcons = true;
+		Info.bFireAndForget = false;
+		Info.bUseThrobber = false;
+		Info.FadeOutDuration = 0.0f;
+		Info.ExpireDuration = 0.0f;
+
+		Info.ButtonDetails.Add(FNotificationButtonInfo(LOCTEXT("UpdateSettings", "Update Settings"), FText(), FSimpleDelegate::CreateLambda([NotificationFuture = NotificationPromise.GetFuture().Share()]() {
+			FModuleManager::LoadModuleChecked<ISettingsModule>("Settings").ShowViewer("Editor", "General", "Global");
+
+			TWeakPtr<SNotificationItem> NotificationPtr = NotificationFuture.Get();
+			if (TSharedPtr<SNotificationItem> Notification = NotificationPtr.Pin())
+			{
+				Notification->SetCompletionState(SNotificationItem::CS_None);
+				Notification->ExpireAndFadeout();
+			}
+		}),
+		SNotificationItem::ECompletionState::CS_Fail));
+
+		TSharedPtr<SNotificationItem> NotificationItem = FSlateNotificationManager::Get().AddNotification(Info);
+		if (NotificationItem.IsValid())
+		{
+			NotificationPromise.SetValue(NotificationItem);
+			NotificationItem->SetCompletionState(SNotificationItem::CS_Fail);
+		}
+	}
+
+	if (DDCProjectSettings->RecommendEveryoneSetupAGlobalS3DDCPath && (Settings->bEnableS3DDC && Settings->GlobalS3DDCPath.Path.IsEmpty()))
+	{
+		TPromise<TWeakPtr<SNotificationItem>> NotificationPromise;
+
+		FNotificationInfo Info(LOCTEXT("SharedProjectS3DDC", "This project recommends you setup the 'Global Local S3 DDC Path', \nso that all copies of this project use the same local S3 DDC cache."));
+		Info.bUseSuccessFailIcons = true;
+		Info.bFireAndForget = false;
+		Info.bUseThrobber = false;
+		Info.FadeOutDuration = 0.0f;
+		Info.ExpireDuration = 0.0f;
+
+		Info.ButtonDetails.Add(FNotificationButtonInfo(LOCTEXT("UpdateSettings", "Update Settings"), FText(), FSimpleDelegate::CreateLambda([NotificationFuture = NotificationPromise.GetFuture().Share()]() {
+			FModuleManager::LoadModuleChecked<ISettingsModule>("Settings").ShowViewer("Editor", "General", "Global");
+
+			TWeakPtr<SNotificationItem> NotificationPtr = NotificationFuture.Get();
+			if (TSharedPtr<SNotificationItem> Notification = NotificationPtr.Pin())
+			{
+				Notification->SetCompletionState(SNotificationItem::CS_None);
+				Notification->ExpireAndFadeout();
+			}
+		}),
+		SNotificationItem::ECompletionState::CS_Fail));
+
+		TSharedPtr<SNotificationItem> NotificationItem = FSlateNotificationManager::Get().AddNotification(Info);
+		if (NotificationItem.IsValid())
+		{
+			NotificationPromise.SetValue(NotificationItem);
+			NotificationItem->SetCompletionState(SNotificationItem::CS_Fail);
+		}
+	}
+
+	return EActiveTimerReturnType::Stop;
+}
+
+FText SDerivedDataStatusBarWidget::GetTitleToolTipText()
+{
+	FTextBuilder DescBuilder;
+
+	DescBuilder.AppendLineFormat(LOCTEXT("GraphNameText", "Graph : {0}"), FText::FromString(GetDerivedDataCache()->GetGraphName()));
+	DescBuilder.AppendLineFormat(LOCTEXT("GraphStatusText", "Status : {0}"), FText::FromString((GetDerivedDataCache()->AnyAsyncRequestsRemaining() ? TEXT("Busy") : TEXT("Idle"))));
+
+	return DescBuilder.ToText();
+}
+
+FText SDerivedDataStatusBarWidget::GetTitleText()
+{
+	return LOCTEXT("DerivedDataToolBarName", "Derived Data");
+}
+
+FText SDerivedDataStatusBarWidget::GetRemoteCacheToolTipText()
+{
+	FTextBuilder DescBuilder;
+
+	DescBuilder.AppendLineFormat(LOCTEXT("RemoteCacheConnectedText", "Remote Cache : {0}\n"), FText::FromString((FDerivedDataInformation::GetHasRemoteCache() ? TEXT("Connected") : TEXT("Unavailable"))));
+
+	const double DownloadedBytesMB = FUnitConversion::Convert(FDerivedDataInformation::GetCacheActivitySizeBytes(true, false), EUnit::Bytes, EUnit::Megabytes);
+	const double UploadedBytesMB = FUnitConversion::Convert(FDerivedDataInformation::GetCacheActivitySizeBytes(false, false), EUnit::Bytes, EUnit::Megabytes);
+
+	DescBuilder.AppendLineFormat(LOCTEXT("RemoteCacheDownloaded", "Downloaded\t: {0} MB"), UploadedBytesMB);
+	DescBuilder.AppendLineFormat(LOCTEXT("RemoteCacheUploaded",	"Uploaded\t: {0} MB"), DownloadedBytesMB);
+	
+	return DescBuilder.ToText();
+}
+
+#undef LOCTEXT_NAMESPACE
\ No newline at end of file
diff --git a/Engine/Source/Editor/DerivedDataEditor/Private/SDerivedDataStatusBar.h b/Engine/Source/Editor/DerivedDataEditor/Private/SDerivedDataStatusBar.h
new file mode 100644
--- /dev/null
+++ b/Engine/Source/Editor/DerivedDataEditor/Private/SDerivedDataStatusBar.h
@@ -0,0 +1,63 @@
+// Copyright Epic Games, Inc. All Rights Reserved.
+
+#pragma once
+
+#include "CoreMinimal.h"
+#include "Widgets/DeclarativeSyntaxSupport.h"
+#include "Widgets/SWidget.h"
+#include "Widgets/SCompoundWidget.h"
+#include "Animation/CurveSequence.h"
+
+class FDerivedDataStatusBarMenuCommands : public TCommands<FDerivedDataStatusBarMenuCommands>
+{
+public:
+
+	FDerivedDataStatusBarMenuCommands();
+
+	virtual void RegisterCommands() override;
+
+private:
+
+
+	static void ChangeSettings_Clicked();
+	static void ViewCacheStatistics_Clicked();
+	static void ViewResourceUsage_Clicked();
+
+
+public:
+
+	TSharedPtr< FUICommandInfo > ChangeSettings;
+	TSharedPtr< FUICommandInfo > ViewResourceUsage;
+	TSharedPtr< FUICommandInfo > ViewCacheStatistics;
+
+	static TSharedRef<FUICommandList> ActionList;
+};
+
+
+class SDerivedDataStatusBarWidget : public SCompoundWidget
+{
+	SLATE_BEGIN_ARGS(SDerivedDataStatusBarWidget) {}
+	SLATE_END_ARGS()
+
+	void Construct(const FArguments& InArgs);
+
+private:
+
+	static FText				GetTitleToolTipText();
+	static FText				GetRemoteCacheToolTipText();
+	static FText				GetTitleText();
+	TSharedRef<SWidget>			CreateStatusBarMenu();
+	EActiveTimerReturnType		UpdateBusyIndicator(double InCurrentTime, float InDeltaTime);
+	EActiveTimerReturnType		UpdateWarnings(double InCurrentTime, float InDeltaTime);
+
+	double LastDDCGetTime = 0;
+	double LastDDCPutTime = 0;
+
+	bool bGetActive = false;
+	bool bPutActive = false;
+	bool bBusy = false;
+
+	FCurveSequence BusyPulseSequence;
+	FCurveSequence FadeGetSequence;
+	FCurveSequence FadePutSequence;
+};
diff --git a/Engine/Source/Editor/DerivedDataEditor/Public/DerivedDataEditorModule.h b/Engine/Source/Editor/DerivedDataEditor/Public/DerivedDataEditorModule.h
new file mode 100644
--- /dev/null
+++ b/Engine/Source/Editor/DerivedDataEditor/Public/DerivedDataEditorModule.h
@@ -0,0 +1,53 @@
+// Copyright Epic Games, Inc. All Rights Reserved.
+
+#pragma once
+
+#include "CoreMinimal.h"
+#include "Modules/ModuleInterface.h"
+#include "Modules/ModuleManager.h"
+
+class SDockTab;
+class SWindow;
+class SDerivedDataCacheSettingsDialog;
+
+/**
+ * The module holding all of the UI related pieces for DerivedData
+ */
+class DERIVEDDATAEDITOR_API FDerivedDataEditorModule : public IModuleInterface
+{
+public:
+
+	/**
+	 * Called right after the module DLL has been loaded and the module object has been created
+	 */
+	virtual void StartupModule() override;
+
+	/**
+	 * Called before the module is unloaded, right before the module object is destroyed.
+	 */
+	virtual void ShutdownModule() override;
+
+	TSharedRef<SWidget>	CreateStatusBarWidget();
+
+	void ShowResourceUsageTab();
+	void ShowCacheStatisticsTab();
+	void ShowSettingsDialog();
+
+private:
+
+	TSharedPtr<SWidget> CreateResourceUsageDialog();
+	TSharedPtr<SWidget> CreateCacheStatisticsDialog();
+
+	TSharedRef<SDockTab> CreateResourceUsageTab(const FSpawnTabArgs& Args);
+	TSharedRef<SDockTab> CreateCacheStatisticsTab(const FSpawnTabArgs& Args);
+
+	void OnSettingsDialogClosed(const TSharedRef<SWindow>& InWindow);
+
+	TWeakPtr<SDockTab> ResourceUsageTab;
+	TWeakPtr<SDockTab> CacheStatisticsTab;
+
+	TSharedPtr<SWindow>	SettingsWindow;
+	TSharedPtr<SDerivedDataCacheSettingsDialog> SettingsDialog;
+};
+
+
diff --git a/Engine/Source/Editor/EditorStyle/Private/StarshipStyle.cpp b/Engine/Source/Editor/EditorStyle/Private/StarshipStyle.cpp
--- a/Engine/Source/Editor/EditorStyle/Private/StarshipStyle.cpp
+++ b/Engine/Source/Editor/EditorStyle/Private/StarshipStyle.cpp
@@ -1138,6 +1138,20 @@ void FStarshipEditorStyle::FStyle::SetupGeneralStyles()
 		);
 	}
 
+	// Derived Data
+	{
+		Set("DerivedData.CacheSettings", new CORE_IMAGE_BRUSH_SVG("Starship/Common/settings", Icon16x16, FStyleColors::Foreground));
+		Set("DerivedData.CacheStatistics", new IMAGE_BRUSH_SVG("Starship/Common/Statistics", Icon16x16, FLinearColor::White));
+		Set("DerivedData.ResourceUsage", new CORE_IMAGE_BRUSH_SVG("Starship/Common/layout-spreadsheet", Icon16x16, FLinearColor::White));
+		
+		Set("DerivedData.Uploading", new CORE_IMAGE_BRUSH_SVG("Starship/Common/arrow-up", Icon16x16, FStyleColors::AccentGreen));
+		Set("DerivedData.Downloading", new CORE_IMAGE_BRUSH_SVG("Starship/Common/arrow-down", Icon16x16, FStyleColors::AccentGreen));
+	
+		Set("DerivedData.RemoteStorage.Connected", new CORE_IMAGE_BRUSH_SVG("Starship/Common/world", Icon16x16, FStyleColors::AccentGreen));
+		Set("DerivedData.RemoteStorage.Disconnected", new CORE_IMAGE_BRUSH_SVG("Starship/Common/world", Icon16x16, FStyleColors::AccentYellow));
+		Set("DerivedData.RemoteStorage.Unavailable", new CORE_IMAGE_BRUSH_SVG("Starship/Common/world", Icon16x16, FStyleColors::AccentGray));
+	}
+
 	// Scene Outliner
 	{
 		Set( "SceneOutliner.NewFolderIcon", new CORE_IMAGE_BRUSH_SVG("Starship/Common/folder-plus", Icon16x16 ) );
diff --git a/Engine/Source/Editor/EditorStyle/Private/StarshipStyle.h b/Engine/Source/Editor/EditorStyle/Private/StarshipStyle.h
--- a/Engine/Source/Editor/EditorStyle/Private/StarshipStyle.h
+++ b/Engine/Source/Editor/EditorStyle/Private/StarshipStyle.h
@@ -72,6 +72,7 @@ public:
 		void SetupMyBlueprintStyles();
 		void SetupStatusBarStyle();
 		void SetupColorPickerStyle();
+		void SetupDerivedDataStyle();
 
 		void SettingsChanged(UObject* ChangedObject, FPropertyChangedEvent& PropertyChangedEvent);
 		void SyncSettings();
diff --git a/Engine/Source/Editor/LevelEditor/LevelEditor.Build.cs b/Engine/Source/Editor/LevelEditor/LevelEditor.Build.cs
--- a/Engine/Source/Editor/LevelEditor/LevelEditor.Build.cs
+++ b/Engine/Source/Editor/LevelEditor/LevelEditor.Build.cs
@@ -84,6 +84,7 @@ public class LevelEditor : ModuleRules
 				"TranslationEditor",
 				"SubobjectEditor",
 				"SubobjectDataInterface",
+				"DerivedDataEditor",
 			}
 		);
 
diff --git a/Engine/Source/Editor/LevelEditor/Private/SLevelEditor.cpp b/Engine/Source/Editor/LevelEditor/Private/SLevelEditor.cpp
--- a/Engine/Source/Editor/LevelEditor/Private/SLevelEditor.cpp
+++ b/Engine/Source/Editor/LevelEditor/Private/SLevelEditor.cpp
@@ -68,8 +68,7 @@
 #include "Elements/Component/ComponentElementLevelEditorCommonActionsCustomization.h"
 #include "Elements/SMInstance/SMInstanceElementId.h"
 #include "Elements/SMInstance/SMInstanceElementLevelEditorSelectionCustomization.h"
-#include "DDC/SDDCStatusIndicator.h"
-#include "Virtualization/SVirtualizationStatusIndicator.h"
+#include "DerivedDataEditorModule.h"
 
 #define LOCTEXT_NAMESPACE "SLevelEditor"
 
@@ -1817,19 +1816,13 @@ void SLevelEditor::RegisterStatusBarTools()
 		}));
 
 
+	FDerivedDataEditorModule& DerivedDataEditorModule = FModuleManager::LoadModuleChecked<FDerivedDataEditorModule>("DerivedDataEditor");
+
 	FToolMenuSection& DDCSection = Menu->AddSection("DDC", FText::GetEmpty(), FToolMenuInsert("Compile", EToolMenuInsertType::Before));
 
 	DDCSection.AddEntry(
-		FToolMenuEntry::InitWidget("DDCStatusIndicator", SNew(SDDCStatusIndicator), FText::GetEmpty(), true, false)
+		FToolMenuEntry::InitWidget("DerivedDatatatusBar", DerivedDataEditorModule.CreateStatusBarWidget(), FText::GetEmpty(), true, false)
 	);
-
-	TSharedPtr<SWidget> VirtualizationWidget = UE::Virtualization::GetVirtualizationStatusIndicator();
-	if (VirtualizationWidget)
-	{
-		FToolMenuSection& MirageSection = Menu->AddSection("Mirage", FText::GetEmpty(), FToolMenuInsert("DDC", EToolMenuInsertType::Before));
-
-		MirageSection.AddEntry(FToolMenuEntry::InitWidget("VirtualizationStatusIndicator", VirtualizationWidget.ToSharedRef(), FText::GetEmpty(), true, false));
-	}	
 }
 
 void SLevelEditor::AddStandaloneLevelViewport( const TSharedRef<SLevelViewport>& LevelViewport )
diff --git a/Engine/Source/Editor/UnrealEd/Private/DDC/SDDCInformation.cpp b/Engine/Source/Editor/UnrealEd/Private/DDC/SDDCInformation.cpp
deleted file mode 100644
diff --git a/Engine/Source/Editor/UnrealEd/Private/DDC/SDDCStatusIndicator.cpp b/Engine/Source/Editor/UnrealEd/Private/DDC/SDDCStatusIndicator.cpp
deleted file mode 100644
diff --git a/Engine/Source/Editor/UnrealEd/Private/EditorEngine.cpp b/Engine/Source/Editor/UnrealEd/Private/EditorEngine.cpp
--- a/Engine/Source/Editor/UnrealEd/Private/EditorEngine.cpp
+++ b/Engine/Source/Editor/UnrealEd/Private/EditorEngine.cpp
@@ -1066,7 +1066,8 @@ void UEditorEngine::Init(IEngineLoop* InEngineLoop)
 			TEXT("StaticMeshEditor"),
 			TEXT("EditorFramework"),
 			TEXT("WorldPartitionEditor"),
-			TEXT("EditorConfig")
+			TEXT("EditorConfig"),
+			TEXT("DerivedDataEditor")
 		};
 
 		FScopedSlowTask ModuleSlowTask(UE_ARRAY_COUNT(ModuleNames));
diff --git a/Engine/Source/Editor/UnrealEd/Public/DDC/SDDCInformation.h b/Engine/Source/Editor/UnrealEd/Public/DDC/SDDCInformation.h
deleted file mode 100644
diff --git a/Engine/Source/Editor/UnrealEd/Public/DDC/SDDCStatusIndicator.h b/Engine/Source/Editor/UnrealEd/Public/DDC/SDDCStatusIndicator.h
deleted file mode 100644
