Additional UI fixes for Derived Data Status Bar
Added new DD background icons and badges to Starship Style
Added back support for DD notifications
Moved DDNotifications into DDEditorModule

diff --git a/Engine/Source/Editor/DerivedDataEditor/Private/DerivedDataCacheNotifications.cpp b/Engine/Source/Editor/DerivedDataEditor/Private/DerivedDataCacheNotifications.cpp
new file mode 100644
--- /dev/null
+++ b/Engine/Source/Editor/DerivedDataEditor/Private/DerivedDataCacheNotifications.cpp
@@ -0,0 +1,108 @@
+// Copyright Epic Games, Inc. All Rights Reserved.
+
+#include "DerivedDataCacheNotifications.h"
+#include "Misc/EngineBuildSettings.h"
+#include "Logging/LogMacros.h"
+#include "Misc/CoreMisc.h"
+#include "Misc/App.h"
+#include "Internationalization/Internationalization.h"
+#include "Widgets/Notifications/SNotificationList.h"
+#include "Widgets/Input/SHyperlink.h"
+#include "Framework/Notifications/NotificationManager.h"
+#include "DerivedDataCacheInterface.h"
+#include "Editor/EditorPerformanceSettings.h"
+
+#define LOCTEXT_NAMESPACE "DerivedDataCacheNotifications"
+
+DEFINE_LOG_CATEGORY_STATIC(DerivedDataCacheNotifications, Log, All);
+
+FDerivedDataCacheNotifications::FDerivedDataCacheNotifications() :
+	bShowSharedDDCNotification(true),
+	bSubscribed(false)
+{
+	Subscribe(true);
+}
+
+FDerivedDataCacheNotifications::~FDerivedDataCacheNotifications()
+{
+	Subscribe(false);
+}
+
+void FDerivedDataCacheNotifications::ClearSharedDDCNotification()
+{
+	// Don't call back into slate if already exiting
+	if (IsEngineExitRequested())
+	{
+		return;
+	}
+
+	if (SharedDDCNotification.IsValid())
+	{
+		SharedDDCNotification.Get()->SetCompletionState(SNotificationItem::CS_None);
+		SharedDDCNotification.Get()->ExpireAndFadeout();
+		SharedDDCNotification.Reset();
+	}
+
+}
+
+void FDerivedDataCacheNotifications::OnDDCNotificationEvent(FDerivedDataCacheInterface::EDDCNotification DDCNotification)
+{
+	// Early out if we have turned off notifications
+	if (!GetDefault<UEditorPerformanceSettings>()->bEnableSharedDDCPerformanceNotifications)
+	{
+		return;
+	}
+
+	if (!bShowSharedDDCNotification || DDCNotification != FDerivedDataCacheInterface::SharedDDCPerformanceNotification)
+	{
+		return;
+	}
+
+	// Only show DDC notifications once per session
+	bShowSharedDDCNotification = false;
+
+	FNotificationInfo Info(NSLOCTEXT("SharedDDCNotification", "SharedDDCNotificationMessage", "Shared Data Cache not in use, performance is impacted."));
+	Info.bFireAndForget = false;
+	Info.bUseThrobber = false;
+	Info.FadeOutDuration = 0.0f;
+	Info.ExpireDuration = 0.0f;
+
+	Info.Hyperlink = FSimpleDelegate::CreateLambda([this]() {
+		const FString DerivedDataCacheUrl = TEXT("https://docs.unrealengine.com/latest/INT/Engine/Basics/DerivedDataCache/");
+		ClearSharedDDCNotification();
+		FPlatformProcess::LaunchURL(*DerivedDataCacheUrl, nullptr, nullptr);
+	});
+
+	Info.HyperlinkText = LOCTEXT("SharedDDCNotificationHyperlink", "View Shared Data Cache Documentation");
+		
+	Info.ButtonDetails.Add(FNotificationButtonInfo(LOCTEXT("SharedDDCNotificationDismiss", "Dismiss"), FText(), FSimpleDelegate::CreateLambda([this]() {
+		ClearSharedDDCNotification();
+	})));
+			
+	SharedDDCNotification = FSlateNotificationManager::Get().AddNotification(Info);
+	if (SharedDDCNotification.IsValid())
+	{
+		SharedDDCNotification.Get()->SetCompletionState(SNotificationItem::CS_Pending);
+	}
+}
+
+void FDerivedDataCacheNotifications::Subscribe(bool bSubscribe)
+{
+	if (bSubscribe != bSubscribed)
+	{
+		FDerivedDataCacheInterface::FOnDDCNotification& DDCNotificationEvent = GetDerivedDataCacheRef().GetDDCNotificationEvent();
+
+		if (bSubscribe)
+		{
+			DDCNotificationEvent.AddRaw(this, &FDerivedDataCacheNotifications::OnDDCNotificationEvent);
+		}
+		else
+		{
+			DDCNotificationEvent.RemoveAll(this);
+		}
+
+		bSubscribed = bSubscribe;
+	}
+}
+
+#undef LOCTEXT_NAMESPACE
diff --git a/Engine/Source/Editor/DerivedDataEditor/Private/DerivedDataCacheNotifications.h b/Engine/Source/Editor/DerivedDataEditor/Private/DerivedDataCacheNotifications.h
new file mode 100644
--- /dev/null
+++ b/Engine/Source/Editor/DerivedDataEditor/Private/DerivedDataCacheNotifications.h
@@ -0,0 +1,39 @@
+// Copyright Epic Games, Inc. All Rights Reserved.
+
+#pragma once
+
+#include "CoreMinimal.h"
+#include "SlateFwd.h"
+#include "UObject/WeakObjectPtr.h"
+#include "DerivedDataCacheInterface.h"
+#include "IDerivedDataCacheNotifications.h"
+
+class SNotificationItem;
+
+class FDerivedDataCacheNotifications : public IDerivedDataCacheNotifications
+{
+public:
+	FDerivedDataCacheNotifications();
+	virtual ~FDerivedDataCacheNotifications();
+
+private: 
+
+	/** DDC data put notification handler */
+	void OnDDCNotificationEvent(FDerivedDataCacheInterface::EDDCNotification DDCNotification);
+	
+	/** Manually clear any presented DDC notifications */
+	void ClearSharedDDCNotification();
+
+	/** Subscribe to the notifactions */
+	void Subscribe(bool bSubscribe);
+
+	/** Whether or not to show the Shared DDC notification */
+	bool bShowSharedDDCNotification;
+
+	/** Whether we are subscribed or not **/
+	bool bSubscribed;
+
+	/** Valid when a DDC notification item is being presented */
+	TSharedPtr<SNotificationItem> SharedDDCNotification;
+};
+
diff --git a/Engine/Source/Editor/DerivedDataEditor/Private/DerivedDataEditorModule.cpp b/Engine/Source/Editor/DerivedDataEditor/Private/DerivedDataEditorModule.cpp
--- a/Engine/Source/Editor/DerivedDataEditor/Private/DerivedDataEditorModule.cpp
+++ b/Engine/Source/Editor/DerivedDataEditor/Private/DerivedDataEditorModule.cpp
@@ -5,10 +5,12 @@
 #include "SDerivedDataStatusBar.h"
 #include "SDerivedDataDialogs.h"
 #include "SDerivedDataCacheSettings.h"
-#include "Widgets/DeclarativeSyntaxSupport.h"
+#include "DerivedDataCacheNotifications.h"
 #include "Modules/ModuleManager.h"
+#include "Framework/Application/SlateApplication.h"
+#include "Widgets/DeclarativeSyntaxSupport.h"
 #include "Widgets/Docking/SDockTab.h"
-
+#include "EditorStyleSet.h"
 #include "WorkspaceMenuStructure.h"
 #include "WorkspaceMenuStructureModule.h"
 
@@ -32,7 +34,7 @@ void FDerivedDataEditorModule::StartupModule()
 		.SetGroup(WorkspaceMenu::GetMenuStructure().GetToolsCategory())
 		.SetIcon(ResourceUsageIcon);
 
-	const FSlateIcon CacheStatisticsIcon(FEditorStyle::GetStyleSetName(), "DerivedData.CacheStatistics");
+	const FSlateIcon CacheStatisticsIcon(FEditorStyle::GetStyleSetName(), "DerivedData.Cache.Statistics");
 
 	FGlobalTabmanager::Get()->RegisterNomadTabSpawner(DerivedDataCacheStatisticsTabName, FOnSpawnTab::CreateRaw(this, &FDerivedDataEditorModule::CreateCacheStatisticsTab))
 		.SetDisplayName(LOCTEXT("DerivedDataCacheStatisticsTabTitle", "Cache Statistics"))
@@ -51,6 +53,8 @@ void FDerivedDataEditorModule::StartupModule()
 #endif // WITH_RELOAD
 
 	FDerivedDataStatusBarMenuCommands::Register();
+
+	DerivedDataCacheNotifications.Reset(new FDerivedDataCacheNotifications);
 }
 
 void FDerivedDataEditorModule::ShutdownModule()
diff --git a/Engine/Source/Editor/DerivedDataEditor/Private/DerivedDataInformation.cpp b/Engine/Source/Editor/DerivedDataEditor/Private/DerivedDataInformation.cpp
--- a/Engine/Source/Editor/DerivedDataEditor/Private/DerivedDataInformation.cpp
+++ b/Engine/Source/Editor/DerivedDataEditor/Private/DerivedDataInformation.cpp
@@ -9,6 +9,12 @@
 
 #define LOCTEXT_NAMESPACE "DerivedDataEditor"
 
+double FDerivedDataInformation::LastGetTime = 0;
+double FDerivedDataInformation::LastPutTime = 0;
+bool FDerivedDataInformation::bIsDownloading = false;
+bool FDerivedDataInformation::bIsUploading = false;
+ERemoteCacheState FDerivedDataInformation::RemoteCacheState= ERemoteCacheState::Unavailable;
+
 double FDerivedDataInformation::GetCacheActivitySizeBytes(bool bGet, bool bLocal)
 {
 	PRAGMA_DISABLE_DEPRECATION_WARNINGS
@@ -151,4 +157,35 @@ bool FDerivedDataInformation::GetHasRemoteCache()
 	return false;
 }
 
+void FDerivedDataInformation::UpdateRemoteCacheState()
+{
+	RemoteCacheState = ERemoteCacheState::Unavailable;
+
+	if ( GetHasRemoteCache() )
+	{
+		const double OldLastGetTime = LastGetTime;
+		const double OldLastPutTime = LastPutTime;
+
+		LastGetTime = FDerivedDataInformation::GetCacheActivityTimeSeconds(true, false);
+		LastPutTime = FDerivedDataInformation::GetCacheActivityTimeSeconds(false, false);
+
+		if (OldLastGetTime != 0.0 && OldLastPutTime != 0.0)
+		{
+			bIsDownloading = OldLastGetTime != LastGetTime;
+			bIsUploading = OldLastPutTime != LastPutTime;
+		}
+
+		if (bIsUploading || bIsDownloading)
+		{
+			RemoteCacheState = ERemoteCacheState::Busy;
+		}
+		else
+		{
+			RemoteCacheState = ERemoteCacheState::Idle;
+		}
+	}
+
+}
+
+
 #undef LOCTEXT_NAMESPACE
\ No newline at end of file
diff --git a/Engine/Source/Editor/DerivedDataEditor/Private/DerivedDataInformation.h b/Engine/Source/Editor/DerivedDataEditor/Private/DerivedDataInformation.h
--- a/Engine/Source/Editor/DerivedDataEditor/Private/DerivedDataInformation.h
+++ b/Engine/Source/Editor/DerivedDataEditor/Private/DerivedDataInformation.h
@@ -5,12 +5,33 @@
 #include "CoreMinimal.h"
 #include "Framework/Commands/Commands.h"
 
+enum class ERemoteCacheState : uint8
+{
+	Idle,
+	Busy,
+	Unavailable,
+	Warning,
+};
+
 class FDerivedDataInformation
 {
 public:
 
-	static double	GetCacheActivityTimeSeconds(bool bGet, bool bLocal);
-	static double	GetCacheActivitySizeBytes(bool bGet, bool bLocal);
-	static bool		GetHasLocalCache();
-	static bool		GetHasRemoteCache();
+	static double				GetCacheActivityTimeSeconds(bool bGet, bool bLocal);
+	static double				GetCacheActivitySizeBytes(bool bGet, bool bLocal);
+	static bool					GetHasLocalCache();
+	static bool					GetHasRemoteCache();
+	static void					UpdateRemoteCacheState();
+	static ERemoteCacheState	GetRemoteCacheState() { return RemoteCacheState; }
+	static bool					IsUploading() { return bIsUploading; }
+	static bool					IsDownloading() { return bIsDownloading; }
+
+private:
+
+	static ERemoteCacheState	RemoteCacheState;
+	static double				LastGetTime;
+	static double				LastPutTime;
+	static bool					bIsUploading;
+	static bool					bIsDownloading;
+
 };
diff --git a/Engine/Source/Editor/DerivedDataEditor/Private/SDerivedDataCacheSettings.cpp b/Engine/Source/Editor/DerivedDataEditor/Private/SDerivedDataCacheSettings.cpp
--- a/Engine/Source/Editor/DerivedDataEditor/Private/SDerivedDataCacheSettings.cpp
+++ b/Engine/Source/Editor/DerivedDataEditor/Private/SDerivedDataCacheSettings.cpp
@@ -3,19 +3,21 @@
 #include "SDerivedDataCacheSettings.h"
 #include "DerivedDataInformation.h"
 #include "DerivedDataCacheInterface.h"
-#include "Widgets/SBoxPanel.h"
-#include "Widgets/Text/STextBlock.h"
+#include "DerivedDataEditorModule.h"
+#include "Widgets/Layout/SBox.h"
+#include "Widgets/Layout/SBorder.h"
 #include "Widgets/Layout/SGridPanel.h"
-#include "Widgets/Images/SImage.h"
 #include "Widgets/Layout/SExpandableArea.h"
-#include "Widgets/Images/SThrobber.h"
-#include "Widgets/Layout/SBox.h"
+#include "Widgets/Text/STextBlock.h"
 #include "Widgets/Input/SEditableTextBox.h"
-#include "SPrimaryButton.h"
+#include "Widgets/Input/SButton.h"
+#include "Widgets/Images/SImage.h"
+#include "Widgets/Images/SThrobber.h"
+#include "Editor/EditorPerformanceSettings.h"
+
+#include "Modules/ModuleManager.h"
 #include "Logging/MessageLog.h"
 #include "MessageLogModule.h"
-#include "Modules/ModuleManager.h"
-
 
 #define LOCTEXT_NAMESPACE "DerivedDataEditor"
 
@@ -53,7 +55,7 @@ void SDerivedDataCacheSettingsDialog::Construct(const FArguments& InArgs)
 						+SVerticalBox::Slot()
 						.Padding(FMargin(0.0f, 0.0f, 0.0f, 10.0f))
 						.VAlign(VAlign_Center)
-						.HAlign(HAlign_Right)
+						.HAlign(HAlign_Left)
 						[
 							SNew(STextBlock)
 							.Text(LOCTEXT("Setting1Label", "Setting1"))
@@ -62,12 +64,21 @@ void SDerivedDataCacheSettingsDialog::Construct(const FArguments& InArgs)
 						+SVerticalBox::Slot()
 						.Padding(FMargin(0.0f, 0.0f, 0.0f, 10.0f))
 						.VAlign(VAlign_Center)
-						.HAlign(HAlign_Right)
+						.HAlign(HAlign_Left)
 						[
 							SNew(STextBlock)
 							.Text(LOCTEXT("Setting2Label", "Setting2"))
 							.ToolTipText( LOCTEXT("Setting2Label_Tooltip", "Setting2") )
-						]		
+						]	
+						+SVerticalBox::Slot()
+						.Padding(FMargin(0.0f, 0.0f, 0.0f, 10.0f))
+						.HAlign(HAlign_Left)
+						[
+							SNew(SCheckBox)
+							.ForegroundColor(FSlateColor::UseForeground())
+							.IsChecked(this, &SDerivedDataCacheSettingsDialog::AreNotificationsEnabled)
+							.OnCheckStateChanged(this, &SDerivedDataCacheSettingsDialog::OnNotifcationsEnabledCheckboxChanged)
+						]
 					]
 					+SHorizontalBox::Slot()
 					.FillWidth(2.0f)
@@ -75,6 +86,7 @@ void SDerivedDataCacheSettingsDialog::Construct(const FArguments& InArgs)
 						SNew(SVerticalBox)
 						+SVerticalBox::Slot()
 						.VAlign(VAlign_Center)
+						.HAlign(HAlign_Left)
 						.Padding(FMargin(0.0f, 0.0f, 0.0f, 10.0f))
 						[
 							SNew(SEditableTextBox)
@@ -86,14 +98,25 @@ void SDerivedDataCacheSettingsDialog::Construct(const FArguments& InArgs)
 						+SVerticalBox::Slot()
 						.Padding(FMargin(0.0f, 0.0f, 0.0f, 10.0f))
 						.VAlign(VAlign_Center)
+						.HAlign(HAlign_Left)
 						[
 							SNew(SEditableTextBox)
 							.Text(this, &SDerivedDataCacheSettingsDialog::GetSetting2Text)
 							.ToolTipText( LOCTEXT("Setting2Label_Tooltip", "Setting2") )
 							.OnTextCommitted(this, &SDerivedDataCacheSettingsDialog::OnSetting2TextCommited)
 							.OnTextChanged(this, &SDerivedDataCacheSettingsDialog::OnSetting2TextCommited, ETextCommit::Default)
-						]		
-					]	
+						]
+						+ SVerticalBox::Slot()
+						.Padding(FMargin(0.0f, 0.0f, 0.0f, 10.0f))
+						.VAlign(VAlign_Center)
+						.HAlign(HAlign_Left)
+						[
+							SNew(STextBlock)
+							.TextStyle(FAppStyle::Get(), "DialogButtonText")
+							.Text(LOCTEXT("EnableDDCNotifications", "Enable Notifications"))
+						]
+					]
+					
 				]
 
 				//+ SVerticalBox::Slot()
@@ -146,7 +169,10 @@ void SDerivedDataCacheSettingsDialog::Construct(const FArguments& InArgs)
 						.AutoWidth()
 						.Padding(FMargin(5.0f, 0.0f))
 						[
-							SNew(SPrimaryButton)
+							SNew(SButton)
+							.VAlign(VAlign_Center)
+							.ButtonStyle(FAppStyle::Get(), "PrimaryButton")
+  							.TextStyle( FAppStyle::Get(), "DialogButtonText" )
 							.Text(LOCTEXT("AcceptSettings", "Accept Settings"))
 							.OnClicked( this, &SDerivedDataCacheSettingsDialog::OnAcceptSettings )
 							.IsEnabled( this, &SDerivedDataCacheSettingsDialog::IsAcceptSettingsEnabled )
@@ -237,8 +263,20 @@ FReply SDerivedDataCacheSettingsDialog::OnAcceptSettings()
 
 FReply SDerivedDataCacheSettingsDialog::OnDisableDerivedDataCache()
 {
+
 	return FReply::Handled();
 }
 
+ECheckBoxState SDerivedDataCacheSettingsDialog::AreNotificationsEnabled() const
+{	
+	return GetDefault<UEditorPerformanceSettings>()->bEnableSharedDDCPerformanceNotifications ? ECheckBoxState::Checked : ECheckBoxState::Unchecked;
+}
+
+void SDerivedDataCacheSettingsDialog::OnNotifcationsEnabledCheckboxChanged(ECheckBoxState NewCheckboxState)
+{
+	UEditorPerformanceSettings* Settings = GetMutableDefault<UEditorPerformanceSettings>();
+	Settings->bEnableSharedDDCPerformanceNotifications = NewCheckboxState == ECheckBoxState::Checked;
+	Settings->PostEditChange();
+}
 
 #undef LOCTEXT_NAMESPACE
\ No newline at end of file
diff --git a/Engine/Source/Editor/DerivedDataEditor/Private/SDerivedDataCacheSettings.h b/Engine/Source/Editor/DerivedDataEditor/Private/SDerivedDataCacheSettings.h
--- a/Engine/Source/Editor/DerivedDataEditor/Private/SDerivedDataCacheSettings.h
+++ b/Engine/Source/Editor/DerivedDataEditor/Private/SDerivedDataCacheSettings.h
@@ -5,6 +5,8 @@
 #include "CoreMinimal.h"
 #include "Widgets/DeclarativeSyntaxSupport.h"
 #include "Widgets/SWidget.h"
+#include "Widgets/SBoxPanel.h"
+#include "Widgets/Input/SCheckBox.h"
 #include "Widgets/SCompoundWidget.h"
 #include "Widgets/SBoxPanel.h"
 
@@ -28,11 +30,14 @@ private:
 
 	void OnSetting1TextCommited(const FText& InText, ETextCommit::Type InCommitType) const;
 	void OnSetting2TextCommited(const FText& InText, ETextCommit::Type InCommitType) const;
+	void OnNotifcationsEnabledCheckboxChanged(ECheckBoxState NewCheckboxState);
 
 	EVisibility GetThrobberVisibility() const;
 	bool IsDerivedDataCacheEnabled() const;
 	bool IsAcceptSettingsEnabled() const;
+	ECheckBoxState AreNotificationsEnabled() const;
 	
 	FReply OnAcceptSettings();	
 	FReply OnDisableDerivedDataCache();
+
 };
diff --git a/Engine/Source/Editor/DerivedDataEditor/Private/SDerivedDataDialogs.cpp b/Engine/Source/Editor/DerivedDataEditor/Private/SDerivedDataDialogs.cpp
--- a/Engine/Source/Editor/DerivedDataEditor/Private/SDerivedDataDialogs.cpp
+++ b/Engine/Source/Editor/DerivedDataEditor/Private/SDerivedDataDialogs.cpp
@@ -12,6 +12,8 @@
 #include "Widgets/Images/SImage.h"
 #include "Styling/StyleColors.h"
 #include "Math/UnitConversion.h"
+#include "Math/BasicMathExpressionEvaluator.h"
+#include "Misc/ExpressionParser.h"
 #include "Framework/Application/SlateApplication.h"
 #include "Internationalization/FastDecimalFormat.h"
 
@@ -558,7 +560,7 @@ TSharedRef<SWidget> SDerivedDataCacheStatisticsDialog::GetGridPanel()
 		.Margin(FMargin(ColumnMargin, RowMargin, 0.0f, TitleMargin))
 		.ColorAndOpacity(FStyleColors::Foreground)
 		.Font(FCoreStyle::GetDefaultFontStyle("Bold", 10))
-		.Justification(ETextJustify::Right)
+		.Justification(ETextJustify::Center)
 		.Text(LOCTEXT("Speed", "Speed"))
 	];
 
@@ -568,7 +570,7 @@ TSharedRef<SWidget> SDerivedDataCacheStatisticsDialog::GetGridPanel()
 		.Margin(FMargin(ColumnMargin, RowMargin, 0.0f, TitleMargin))
 		.ColorAndOpacity(FStyleColors::Foreground)
 		.Font(FCoreStyle::GetDefaultFontStyle("Bold", 10))
-		.Justification(ETextJustify::Right)
+		.Justification(ETextJustify::Center)
 		.Text(LOCTEXT("HitPercentage", "Hit%"))
 	];
 
@@ -578,7 +580,7 @@ TSharedRef<SWidget> SDerivedDataCacheStatisticsDialog::GetGridPanel()
 		.Margin(FMargin(ColumnMargin, RowMargin, 0.0f, TitleMargin))
 		.ColorAndOpacity(FStyleColors::Foreground)
 		.Font(FCoreStyle::GetDefaultFontStyle("Bold", 10))
-		.Justification(ETextJustify::Right)
+		.Justification(ETextJustify::Center)
 		.Text(LOCTEXT("Read", "Read"))
 	];
 
@@ -588,18 +590,10 @@ TSharedRef<SWidget> SDerivedDataCacheStatisticsDialog::GetGridPanel()
 		.Margin(FMargin(ColumnMargin, RowMargin, 0.0f, TitleMargin))
 		.ColorAndOpacity(FStyleColors::Foreground)
 		.Font(FCoreStyle::GetDefaultFontStyle("Bold", 10))
-		.Justification(ETextJustify::Right)
+		.Justification(ETextJustify::Center)
 		.Text(LOCTEXT("Write", "Write"))
 		];
 
-	Panel->AddSlot(5, Row)
-	[
-		SNew(SImage)
-		.Image(FAppStyle::Get().GetBrush("Icons.Edit"))
-		.ColorAndOpacity(FStyleColors::Foreground)
-		
-	];
-
 	Row++;
 
 	double SumTotalGetMB =0.0;
@@ -680,14 +674,6 @@ TSharedRef<SWidget> SDerivedDataCacheStatisticsDialog::GetGridPanel()
 				.Text(FText::FromString(SingleDecimalFormat(TotalPutMB) + TEXT(" MB")))
 			];
 
-			Panel->AddSlot(5, Row)
-			.HAlign(HAlign_Center)
-			[
-				SNew(SImage)
-				.Image(FAppStyle::Get().GetBrush("Icons.Check"))
-				.Visibility(Backend->IsWritable() ? EVisibility::Visible : EVisibility::Hidden)
-			];
-
 			Row++;
 		}
 	}
diff --git a/Engine/Source/Editor/DerivedDataEditor/Private/SDerivedDataStatusBar.cpp b/Engine/Source/Editor/DerivedDataEditor/Private/SDerivedDataStatusBar.cpp
--- a/Engine/Source/Editor/DerivedDataEditor/Private/SDerivedDataStatusBar.cpp
+++ b/Engine/Source/Editor/DerivedDataEditor/Private/SDerivedDataStatusBar.cpp
@@ -96,7 +96,7 @@ TSharedRef<SWidget> SDerivedDataStatusBarWidget::CreateStatusBarMenu()
 			FDerivedDataStatusBarMenuCommands::Get().ChangeSettings,
 			TAttribute<FText>(),
 			TAttribute<FText>(),
-			FSlateIcon(FEditorStyle::GetStyleSetName(), "DerivedData.CacheSettings")
+			FSlateIcon(FEditorStyle::GetStyleSetName(), "DerivedData.Cache.Settings")
 		);
 	}
 
@@ -107,7 +107,7 @@ TSharedRef<SWidget> SDerivedDataStatusBarWidget::CreateStatusBarMenu()
 			FDerivedDataStatusBarMenuCommands::Get().ViewCacheStatistics,
 			TAttribute<FText>(),
 			TAttribute<FText>(),
-			FSlateIcon(FEditorStyle::GetStyleSetName(), "DerivedData.CacheStatistics")
+			FSlateIcon(FEditorStyle::GetStyleSetName(), "DerivedData.Cache.Statistics")
 		);
 
 		Section.AddMenuEntry(
@@ -123,17 +123,6 @@ TSharedRef<SWidget> SDerivedDataStatusBarWidget::CreateStatusBarMenu()
 
 void SDerivedDataStatusBarWidget::Construct(const FArguments& InArgs)
 {	
-	/*FDerivedDataCacheInterface::FOnDDCNotification& DDCNotificationEvent = GetDerivedDataCacheRef().GetDDCNotificationEvent();
-
-	if (bSubscribe)
-	{
-		DDCNotificationEvent.AddRaw(this, &FDDCNotifications::OnDDCNotificationEvent);
-	}
-	else
-	{
-		DDCNotificationEvent.RemoveAll(this);
-	}*/
-
 	BusyPulseSequence = FCurveSequence(0.f, 1.0f, ECurveEaseFunction::QuadInOut);
 	FadeGetSequence = FCurveSequence(0.f, 0.5f, ECurveEaseFunction::Linear);
 	FadePutSequence = FCurveSequence(0.f, 0.5f, ECurveEaseFunction::Linear);
@@ -160,7 +149,18 @@ void SDerivedDataStatusBarWidget::Construct(const FArguments& InArgs)
 				.VAlign(VAlign_Top)
 				[
 					SNew(SImage)
-					.Image_Lambda([this] {  return FDerivedDataInformation::GetHasRemoteCache() ? FAppStyle::Get().GetBrush("DerivedData.RemoteStorage.Connected") : FAppStyle::Get().GetBrush("DerivedData.RemoteStorage.Unavailable"); })			
+					.ColorAndOpacity(FSlateColor::UseForeground())
+					.Image_Lambda([this] { return GetRemoteCacheStateBackgroundIcon();  })
+					.ToolTipText_Lambda([this] { return GetRemoteCacheToolTipText(); })
+				]
+
+				+ SOverlay::Slot()
+				.HAlign(HAlign_Center)
+				.VAlign(VAlign_Top)
+				[
+					SNew(SImage)
+					.ColorAndOpacity(FSlateColor::UseForeground())
+					.Image_Lambda([this] { return GetRemoteCacheStateBadgeIcon();  })
 					.ToolTipText_Lambda([this] { return GetRemoteCacheToolTipText(); })
 				]
 
@@ -170,8 +170,8 @@ void SDerivedDataStatusBarWidget::Construct(const FArguments& InArgs)
 				.Padding(0, 0, 4, 4)
 				[
 					SNew(SImage)
-					.Image(FAppStyle::Get().GetBrush("DerivedData.Uploading"))
-					.ColorAndOpacity_Lambda([this] { return bPutActive ? FLinearColor::White.CopyWithNewOpacity(0.5f + (0.5f * FMath::MakePulsatingValue(FadePutSequence.GetLerp(), 1))) : FLinearColor(0,0,0,0); })
+					.Image(FAppStyle::Get().GetBrush("DerivedData.RemoteCache.Uploading"))
+					.ColorAndOpacity_Lambda([this] { return FDerivedDataInformation::IsUploading() ? FLinearColor::White.CopyWithNewOpacity(0.5f + (0.5f * FMath::MakePulsatingValue(FadePutSequence.GetLerp(), 1))) : FLinearColor(0,0,0,0); })
 					.ToolTipText_Lambda([this] { return GetRemoteCacheToolTipText(); })
 				]
 
@@ -181,10 +181,10 @@ void SDerivedDataStatusBarWidget::Construct(const FArguments& InArgs)
 				.Padding(4, 4, 0, 0)
 				[
 					SNew(SImage)
-					.Image(FAppStyle::Get().GetBrush("DerivedData.Downloading"))
-					.ColorAndOpacity_Lambda([this] { return bGetActive ? FLinearColor::White.CopyWithNewOpacity(0.5f + (0.5f * FMath::MakePulsatingValue(FadeGetSequence.GetLerp(), 1))) : FLinearColor(0, 0, 0, 0); })
+					.Image(FAppStyle::Get().GetBrush("DerivedData.RemoteCache.Downloading"))
+					.ColorAndOpacity_Lambda([this] { return FDerivedDataInformation::IsDownloading() ? FLinearColor::White.CopyWithNewOpacity(0.5f + (0.5f * FMath::MakePulsatingValue(FadeGetSequence.GetLerp(), 1))) : FLinearColor(0, 0, 0, 0); })
 					.ToolTipText_Lambda([this] { return GetRemoteCacheToolTipText(); })
-				]
+				]				
 			]
 
 			+ SHorizontalBox::Slot()
@@ -200,26 +200,18 @@ void SDerivedDataStatusBarWidget::Construct(const FArguments& InArgs)
 		.OnGetMenuContent(FOnGetContent::CreateRaw(this, &SDerivedDataStatusBarWidget::CreateStatusBarMenu))
 	];
 
-	LastDDCGetTime = FDerivedDataInformation::GetCacheActivityTimeSeconds(true, false);
-	LastDDCPutTime = FDerivedDataInformation::GetCacheActivityTimeSeconds(false, false);
 	RegisterActiveTimer(0.5f, FWidgetActiveTimerDelegate::CreateSP(this, &SDerivedDataStatusBarWidget::UpdateBusyIndicator));
 	RegisterActiveTimer(5.0f, FWidgetActiveTimerDelegate::CreateSP(this, &SDerivedDataStatusBarWidget::UpdateWarnings));
 }
 
 EActiveTimerReturnType SDerivedDataStatusBarWidget::UpdateBusyIndicator(double InCurrentTime, float InDeltaTime)
-{
-	const double OldLastDDCGetTime = LastDDCGetTime;
-	const double OldLastDDCPutTime = LastDDCPutTime;
-
-	LastDDCGetTime = FDerivedDataInformation::GetCacheActivityTimeSeconds(true, false);
-	LastDDCPutTime = FDerivedDataInformation::GetCacheActivityTimeSeconds(false, false);
+{	
+	FDerivedDataInformation::UpdateRemoteCacheState();
 
-	bGetActive = OldLastDDCGetTime != LastDDCGetTime;
-	bPutActive = OldLastDDCPutTime != LastDDCPutTime;
 	bBusy = GetDerivedDataCache()->AnyAsyncRequestsRemaining();
 
-	FadeGetSequence.PlayRelative(this->AsShared(), bGetActive);
-	FadePutSequence.PlayRelative(this->AsShared(), bPutActive);
+	FadeGetSequence.PlayRelative(this->AsShared(), FDerivedDataInformation::IsDownloading());
+	FadePutSequence.PlayRelative(this->AsShared(), FDerivedDataInformation::IsUploading());
 
 	if (bBusy)
 	{
@@ -307,7 +299,7 @@ EActiveTimerReturnType SDerivedDataStatusBarWidget::UpdateWarnings(double InCurr
 	return EActiveTimerReturnType::Stop;
 }
 
-FText SDerivedDataStatusBarWidget::GetTitleToolTipText()
+FText SDerivedDataStatusBarWidget::GetTitleToolTipText() const
 {
 	FTextBuilder DescBuilder;
 
@@ -317,12 +309,12 @@ FText SDerivedDataStatusBarWidget::GetTitleToolTipText()
 	return DescBuilder.ToText();
 }
 
-FText SDerivedDataStatusBarWidget::GetTitleText()
+FText SDerivedDataStatusBarWidget::GetTitleText() const
 {
 	return LOCTEXT("DerivedDataToolBarName", "Derived Data");
 }
 
-FText SDerivedDataStatusBarWidget::GetRemoteCacheToolTipText()
+FText SDerivedDataStatusBarWidget::GetRemoteCacheToolTipText() const
 {
 	FTextBuilder DescBuilder;
 
@@ -331,10 +323,72 @@ FText SDerivedDataStatusBarWidget::GetRemoteCacheToolTipText()
 	const double DownloadedBytesMB = FUnitConversion::Convert(FDerivedDataInformation::GetCacheActivitySizeBytes(true, false), EUnit::Bytes, EUnit::Megabytes);
 	const double UploadedBytesMB = FUnitConversion::Convert(FDerivedDataInformation::GetCacheActivitySizeBytes(false, false), EUnit::Bytes, EUnit::Megabytes);
 
-	DescBuilder.AppendLineFormat(LOCTEXT("RemoteCacheDownloaded", "Downloaded\t: {0} MB"), UploadedBytesMB);
-	DescBuilder.AppendLineFormat(LOCTEXT("RemoteCacheUploaded",	"Uploaded\t: {0} MB"), DownloadedBytesMB);
+	DescBuilder.AppendLineFormat(LOCTEXT("RemoteCacheDownloaded", "Downloaded\t: {0} MB"), DownloadedBytesMB);
+	DescBuilder.AppendLineFormat(LOCTEXT("RemoteCacheUploaded",	"Uploaded\t: {0} MB"), UploadedBytesMB);
 	
 	return DescBuilder.ToText();
 }
 
+const FSlateBrush* SDerivedDataStatusBarWidget::GetRemoteCacheStateBackgroundIcon() const
+{
+	switch ( FDerivedDataInformation::GetRemoteCacheState())
+	{
+		case ERemoteCacheState::Idle :
+		{
+			return FAppStyle::Get().GetBrush("DerivedData.RemoteCache.IdleBG");
+			break;
+		}
+
+		case ERemoteCacheState::Busy :
+		{
+			return FAppStyle::Get().GetBrush("DerivedData.RemoteCache.BusyBG");
+			break;
+		}	
+
+		case ERemoteCacheState::Unavailable:
+		{
+			return FAppStyle::Get().GetBrush("DerivedData.RemoteCache.UnavailableBG");
+			break;
+		}
+
+		default:
+		case ERemoteCacheState::Warning:
+		{
+			return FAppStyle::Get().GetBrush("DerivedData.RemoteCache.WarningBG");
+			break;
+		}
+	}
+}
+
+const FSlateBrush* SDerivedDataStatusBarWidget::GetRemoteCacheStateBadgeIcon() const
+{
+	switch ( FDerivedDataInformation::GetRemoteCacheState())
+	{
+		case ERemoteCacheState::Idle :
+		{
+			return FAppStyle::Get().GetBrush("DerivedData.RemoteCache.Idle");
+			break;
+		}
+
+		case ERemoteCacheState::Busy :
+		{
+			return FAppStyle::Get().GetBrush("DerivedData.RemoteCache.Busy");
+			break;
+		}	
+
+		case ERemoteCacheState::Unavailable:
+		{
+			return FAppStyle::Get().GetBrush("DerivedData.RemoteCache.Unavailable");
+			break;
+		}
+
+		default:
+		case ERemoteCacheState::Warning:
+		{
+			return FAppStyle::Get().GetBrush("DerivedData.RemoteCache.Warning");
+			break;
+		}
+	}
+}
+
 #undef LOCTEXT_NAMESPACE
\ No newline at end of file
diff --git a/Engine/Source/Editor/DerivedDataEditor/Private/SDerivedDataStatusBar.h b/Engine/Source/Editor/DerivedDataEditor/Private/SDerivedDataStatusBar.h
--- a/Engine/Source/Editor/DerivedDataEditor/Private/SDerivedDataStatusBar.h
+++ b/Engine/Source/Editor/DerivedDataEditor/Private/SDerivedDataStatusBar.h
@@ -7,6 +7,9 @@
 #include "Widgets/SWidget.h"
 #include "Widgets/SCompoundWidget.h"
 #include "Animation/CurveSequence.h"
+#include "Framework/Commands/Commands.h"
+
+class FUICommandList;
 
 class FDerivedDataStatusBarMenuCommands : public TCommands<FDerivedDataStatusBarMenuCommands>
 {
@@ -43,9 +46,11 @@ class SDerivedDataStatusBarWidget : public SCompoundWidget
 
 private:
 
-	static FText				GetTitleToolTipText();
-	static FText				GetRemoteCacheToolTipText();
-	static FText				GetTitleText();
+	FText						GetTitleToolTipText() const;
+	FText						GetRemoteCacheToolTipText() const;
+	FText						GetTitleText() const;
+	const FSlateBrush*			GetRemoteCacheStateBackgroundIcon() const;
+	const FSlateBrush*			GetRemoteCacheStateBadgeIcon() const;
 	TSharedRef<SWidget>			CreateStatusBarMenu();
 	EActiveTimerReturnType		UpdateBusyIndicator(double InCurrentTime, float InDeltaTime);
 	EActiveTimerReturnType		UpdateWarnings(double InCurrentTime, float InDeltaTime);
diff --git a/Engine/Source/Editor/DerivedDataEditor/Public/DerivedDataEditorModule.h b/Engine/Source/Editor/DerivedDataEditor/Public/DerivedDataEditorModule.h
--- a/Engine/Source/Editor/DerivedDataEditor/Public/DerivedDataEditorModule.h
+++ b/Engine/Source/Editor/DerivedDataEditor/Public/DerivedDataEditorModule.h
@@ -5,10 +5,13 @@
 #include "CoreMinimal.h"
 #include "Modules/ModuleInterface.h"
 #include "Modules/ModuleManager.h"
+#include "IDerivedDataCacheNotifications.h"
 
 class SDockTab;
 class SWindow;
+class SWidget;
 class SDerivedDataCacheSettingsDialog;
+class FSpawnTabArgs;
 
 /**
  * The module holding all of the UI related pieces for DerivedData
@@ -28,6 +31,7 @@ public:
 	virtual void ShutdownModule() override;
 
 	TSharedRef<SWidget>	CreateStatusBarWidget();
+	IDerivedDataCacheNotifications& GetCacheNotifcations() { return *DerivedDataCacheNotifications; }
 
 	void ShowResourceUsageTab();
 	void ShowCacheStatisticsTab();
@@ -48,6 +52,7 @@ private:
 
 	TSharedPtr<SWindow>	SettingsWindow;
 	TSharedPtr<SDerivedDataCacheSettingsDialog> SettingsDialog;
+	TUniquePtr<IDerivedDataCacheNotifications>	DerivedDataCacheNotifications;
 };
 
 
diff --git a/Engine/Source/Editor/DerivedDataEditor/Public/IDerivedDataCacheNotifications.h b/Engine/Source/Editor/DerivedDataEditor/Public/IDerivedDataCacheNotifications.h
new file mode 100644
--- /dev/null
+++ b/Engine/Source/Editor/DerivedDataEditor/Public/IDerivedDataCacheNotifications.h
@@ -0,0 +1,17 @@
+// Copyright Epic Games, Inc. All Rights Reserved.
+
+#pragma once
+
+#include "CoreMinimal.h"
+
+class DERIVEDDATAEDITOR_API IDerivedDataCacheNotifications
+{
+public:
+	IDerivedDataCacheNotifications()
+	{
+	}
+
+	virtual ~IDerivedDataCacheNotifications()
+	{
+	}
+};
\ No newline at end of file
diff --git a/Engine/Source/Editor/EditorStyle/Private/StarshipStyle.cpp b/Engine/Source/Editor/EditorStyle/Private/StarshipStyle.cpp
--- a/Engine/Source/Editor/EditorStyle/Private/StarshipStyle.cpp
+++ b/Engine/Source/Editor/EditorStyle/Private/StarshipStyle.cpp
@@ -1140,16 +1140,24 @@ void FStarshipEditorStyle::FStyle::SetupGeneralStyles()
 
 	// Derived Data
 	{
-		Set("DerivedData.CacheSettings", new CORE_IMAGE_BRUSH_SVG("Starship/Common/settings", Icon16x16, FStyleColors::Foreground));
-		Set("DerivedData.CacheStatistics", new IMAGE_BRUSH_SVG("Starship/Common/Statistics", Icon16x16, FLinearColor::White));
-		Set("DerivedData.ResourceUsage", new CORE_IMAGE_BRUSH_SVG("Starship/Common/layout-spreadsheet", Icon16x16, FLinearColor::White));
+		Set("DerivedData.Cache.Settings", new IMAGE_BRUSH_SVG("Starship/DerivedData/DD_Cache_Settings", Icon16x16));
+		Set("DerivedData.Cache.Statistics", new IMAGE_BRUSH_SVG("Starship/DerivedData/DD_Cache_Statistics", Icon16x16));
+		Set("DerivedData.ResourceUsage", new IMAGE_BRUSH_SVG("Starship/DerivedData/DD_Resource_Usage", Icon16x16));
 		
-		Set("DerivedData.Uploading", new CORE_IMAGE_BRUSH_SVG("Starship/Common/arrow-up", Icon16x16, FStyleColors::AccentGreen));
-		Set("DerivedData.Downloading", new CORE_IMAGE_BRUSH_SVG("Starship/Common/arrow-down", Icon16x16, FStyleColors::AccentGreen));
-	
-		Set("DerivedData.RemoteStorage.Connected", new CORE_IMAGE_BRUSH_SVG("Starship/Common/world", Icon16x16, FStyleColors::AccentGreen));
-		Set("DerivedData.RemoteStorage.Disconnected", new CORE_IMAGE_BRUSH_SVG("Starship/Common/world", Icon16x16, FStyleColors::AccentYellow));
-		Set("DerivedData.RemoteStorage.Unavailable", new CORE_IMAGE_BRUSH_SVG("Starship/Common/world", Icon16x16, FStyleColors::AccentGray));
+		Set("DerivedData.RemoteCache.Uploading", new IMAGE_BRUSH_SVG("Starship/DerivedData/DD_Upload", Icon16x16, EStyleColor::AccentGreen));
+		Set("DerivedData.RemoteCache.Downloading", new IMAGE_BRUSH_SVG("Starship/DerivedData/DD_Download", Icon16x16, EStyleColor::AccentBlue));
+
+		Set("DerivedData.RemoteCache.BusyBG", new IMAGE_BRUSH_SVG("Starship/DerivedData/DD_RemoteCache_UpDownBG", Icon16x16));
+		Set("DerivedData.RemoteCache.Busy", new IMAGE_BRUSH_SVG("Starship/DerivedData/DD_RemoteCache_UpDown", Icon16x16));
+
+		Set("DerivedData.RemoteCache.IdleBG", new IMAGE_BRUSH_SVG("Starship/DerivedData/DD_RemoteCache_IdleBG", Icon16x16));
+		Set("DerivedData.RemoteCache.Idle", new IMAGE_BRUSH_SVG("Starship/DerivedData/DD_RemoteCache_Idle", Icon16x16, EStyleColor::Success));
+
+		Set("DerivedData.RemoteCache.WarningBG", new IMAGE_BRUSH_SVG("Starship/DerivedData/DD_RemoteCache_WarningBG", Icon16x16));
+		Set("DerivedData.RemoteCache.Warning", new IMAGE_BRUSH_SVG("Starship/DerivedData/DD_RemoteCache_Warning", Icon16x16, EStyleColor::Warning));
+
+		Set("DerivedData.RemoteCache.UnavailableBG", new IMAGE_BRUSH_SVG("Starship/DerivedData/DD_RemoteCache_Unavailable", Icon16x16));
+		Set("DerivedData.RemoteCache.Unavailable", new IMAGE_BRUSH_SVG("Starship/DerivedData/DD_RemoteCache_Unavailable", Icon16x16));
 	}
 
 	// Scene Outliner
diff --git a/Engine/Source/Editor/UnrealEd/Classes/Editor/UnrealEdEngine.h b/Engine/Source/Editor/UnrealEd/Classes/Editor/UnrealEdEngine.h
--- a/Engine/Source/Editor/UnrealEd/Classes/Editor/UnrealEdEngine.h
+++ b/Engine/Source/Editor/UnrealEd/Classes/Editor/UnrealEdEngine.h
@@ -10,7 +10,6 @@
 #include "Editor/EditorEngine.h"
 #include "IPackageAutoSaver.h"
 #include "ISourceControlProvider.h"
-#include "IDDCNotifications.h"
 #include "ComponentVisualizer.h"
 #include "ComponentVisualizerManager.h"
 #include "UnrealEdEngine.generated.h"
@@ -892,12 +891,6 @@ public:
 		return *PackageAutoSaver;
 	}
 
-	/** @return The DDC notifications instance used by the editor */
-	IDDCNotifications& GetDDCNotifications() const
-	{
-		return *DDCNotifications;
-	}
-
 	/**
 	 * Exec command handlers
 	 */
@@ -947,9 +940,6 @@ protected:
 	/** The package auto-saver instance used by the editor */
 	TUniquePtr<IPackageAutoSaver> PackageAutoSaver;
 
-	/** The DDC notifications instance used by the editor */
-	TUniquePtr<IDDCNotifications> DDCNotifications;
-
 	/**
 	 * The list of visualizers to draw when selection changes
 	 */
diff --git a/Engine/Source/Editor/UnrealEd/Private/DDCNotifications.cpp b/Engine/Source/Editor/UnrealEd/Private/DDCNotifications.cpp
deleted file mode 100644
diff --git a/Engine/Source/Editor/UnrealEd/Private/DDCNotifications.h b/Engine/Source/Editor/UnrealEd/Private/DDCNotifications.h
deleted file mode 100644
diff --git a/Engine/Source/Editor/UnrealEd/Private/UnrealEdEngine.cpp b/Engine/Source/Editor/UnrealEd/Private/UnrealEdEngine.cpp
--- a/Engine/Source/Editor/UnrealEd/Private/UnrealEdEngine.cpp
+++ b/Engine/Source/Editor/UnrealEd/Private/UnrealEdEngine.cpp
@@ -53,7 +53,6 @@
 #include "StatsViewerModule.h"
 #include "SnappingUtils.h"
 #include "PackageAutoSaver.h"
-#include "DDCNotifications.h"
 #include "PerformanceMonitor.h"
 #include "BSPOps.h"
 #include "SourceCodeNavigation.h"
@@ -136,8 +135,6 @@ void UUnrealEdEngine::Init(IEngineLoop* InEngineLoop)
 	// Iterate over all always fully loaded packages and load them.
 	if (!IsRunningCommandlet())
 	{
-		DDCNotifications.Reset(new FDDCNotifications);
-
 		for( int32 PackageNameIndex=0; PackageNameIndex<PackagesToBeFullyLoadedAtStartup.Num(); PackageNameIndex++ )
 		{
 			const FString& PackageName = PackagesToBeFullyLoadedAtStartup[PackageNameIndex];
@@ -430,8 +427,6 @@ void UUnrealEdEngine::FinishDestroy()
 		PackageAutoSaver.Reset();
 	}
 
-	DDCNotifications.Reset();
-
 	if( PerformanceMonitor )
 	{
 		delete PerformanceMonitor;
diff --git a/Engine/Source/Editor/UnrealEd/Public/IDDCNotifications.h b/Engine/Source/Editor/UnrealEd/Public/IDDCNotifications.h
deleted file mode 100644
