Small UI fix for up/down arrows

diff --git a/Engine/Source/Editor/DerivedDataEditor/Private/SDerivedDataStatusBar.cpp b/Engine/Source/Editor/DerivedDataEditor/Private/SDerivedDataStatusBar.cpp
--- a/Engine/Source/Editor/DerivedDataEditor/Private/SDerivedDataStatusBar.cpp
+++ b/Engine/Source/Editor/DerivedDataEditor/Private/SDerivedDataStatusBar.cpp
@@ -167,7 +167,6 @@ void SDerivedDataStatusBarWidget::Construct(const FArguments& InArgs)
 				+ SOverlay::Slot()
 				.HAlign(HAlign_Center)
 				.VAlign(VAlign_Top)
-				.Padding(0, 0, 4, 4)
 				[
 					SNew(SImage)
 					.Image(FAppStyle::Get().GetBrush("DerivedData.RemoteCache.Uploading"))
@@ -178,7 +177,6 @@ void SDerivedDataStatusBarWidget::Construct(const FArguments& InArgs)
 				+ SOverlay::Slot()
 				.HAlign(HAlign_Center)
 				.VAlign(VAlign_Bottom)
-				.Padding(4, 4, 0, 0)
 				[
 					SNew(SImage)
 					.Image(FAppStyle::Get().GetBrush("DerivedData.RemoteCache.Downloading"))
