Added Status information to the ToolTip text

diff --git a/Engine/Source/Editor/DerivedDataEditor/Private/DerivedDataInformation.cpp b/Engine/Source/Editor/DerivedDataEditor/Private/DerivedDataInformation.cpp
--- a/Engine/Source/Editor/DerivedDataEditor/Private/DerivedDataInformation.cpp
+++ b/Engine/Source/Editor/DerivedDataEditor/Private/DerivedDataInformation.cpp
@@ -108,31 +108,6 @@ double FDerivedDataInformation::GetCacheActivityTimeSeconds(bool bGet, bool bLoc
 	return (double)TotalCycles * FPlatformTime::GetSecondsPerCycle();
 }
 
-bool FDerivedDataInformation::GetHasLocalCache()
-{
-	PRAGMA_DISABLE_DEPRECATION_WARNINGS
-		TSharedRef<FDerivedDataCacheStatsNode> RootUsage = GetDerivedDataCache()->GatherUsageStats();
-	PRAGMA_ENABLE_DEPRECATION_WARNINGS
-		TArray<TSharedRef<const FDerivedDataCacheStatsNode>> LeafUsageStats;
-	RootUsage->ForEachDescendant([&LeafUsageStats](TSharedRef<const FDerivedDataCacheStatsNode> Node) {
-		if (Node->Children.Num() == 0)
-		{
-			LeafUsageStats.Add(Node);
-		}
-		});
-
-	for (int32 Index = 0; Index < LeafUsageStats.Num(); Index++)
-	{
-		const FDerivedDataBackendInterface* Backend = LeafUsageStats[Index]->GetBackendInterface();
-
-		if (Backend->GetSpeedClass() == FDerivedDataBackendInterface::ESpeedClass::Local)
-			return true;
-	}
-
-	return false;
-}
-
-
 bool FDerivedDataInformation::GetHasRemoteCache()
 {
 	PRAGMA_DISABLE_DEPRECATION_WARNINGS
@@ -188,4 +163,36 @@ void FDerivedDataInformation::UpdateRemoteCacheState()
 }
 
 
+FText FDerivedDataInformation::GetRemoteCacheStateAsText()
+{
+	switch (FDerivedDataInformation::GetRemoteCacheState())
+	{
+	case ERemoteCacheState::Idle:
+	{
+		return FText::FromString(TEXT("Idle"));
+		break;
+	}
+
+	case ERemoteCacheState::Busy:
+	{
+		return FText::FromString(TEXT("Busy"));
+		break;
+	}
+
+	case ERemoteCacheState::Unavailable:
+	{
+		return FText::FromString(TEXT("Unavailable"));
+		break;
+	}
+
+	default:
+	case ERemoteCacheState::Warning:
+	{
+		return FText::FromString(TEXT("Warning"));
+		break;
+	}
+	}
+}
+
+
 #undef LOCTEXT_NAMESPACE
\ No newline at end of file
diff --git a/Engine/Source/Editor/DerivedDataEditor/Private/DerivedDataInformation.h b/Engine/Source/Editor/DerivedDataEditor/Private/DerivedDataInformation.h
--- a/Engine/Source/Editor/DerivedDataEditor/Private/DerivedDataInformation.h
+++ b/Engine/Source/Editor/DerivedDataEditor/Private/DerivedDataInformation.h
@@ -19,10 +19,10 @@ public:
 
 	static double				GetCacheActivityTimeSeconds(bool bGet, bool bLocal);
 	static double				GetCacheActivitySizeBytes(bool bGet, bool bLocal);
-	static bool					GetHasLocalCache();
 	static bool					GetHasRemoteCache();
-	static void					UpdateRemoteCacheState();
 	static ERemoteCacheState	GetRemoteCacheState() { return RemoteCacheState; }
+	static FText				GetRemoteCacheStateAsText();
+	static void					UpdateRemoteCacheState();
 	static bool					IsUploading() { return bIsUploading; }
 	static bool					IsDownloading() { return bIsDownloading; }
 
diff --git a/Engine/Source/Editor/DerivedDataEditor/Private/SDerivedDataStatusBar.cpp b/Engine/Source/Editor/DerivedDataEditor/Private/SDerivedDataStatusBar.cpp
--- a/Engine/Source/Editor/DerivedDataEditor/Private/SDerivedDataStatusBar.cpp
+++ b/Engine/Source/Editor/DerivedDataEditor/Private/SDerivedDataStatusBar.cpp
@@ -302,8 +302,7 @@ FText SDerivedDataStatusBarWidget::GetTitleToolTipText() const
 	FTextBuilder DescBuilder;
 
 	DescBuilder.AppendLineFormat(LOCTEXT("GraphNameText", "Graph : {0}"), FText::FromString(GetDerivedDataCache()->GetGraphName()));
-	DescBuilder.AppendLineFormat(LOCTEXT("GraphStatusText", "Status : {0}"), FText::FromString((GetDerivedDataCache()->AnyAsyncRequestsRemaining() ? TEXT("Busy") : TEXT("Idle"))));
-
+	
 	return DescBuilder.ToText();
 }
 
@@ -312,12 +311,45 @@ FText SDerivedDataStatusBarWidget::GetTitleText() const
 	return LOCTEXT("DerivedDataToolBarName", "Derived Data");
 }
 
+FText GetRemoteCacheStateAsText()
+{
+	switch ( FDerivedDataInformation::GetRemoteCacheState())
+	{
+		case ERemoteCacheState::Idle :
+		{
+			return FText::FromString(TEXT("Idle"));
+			break;
+		}
+
+		case ERemoteCacheState::Busy :
+		{
+			return FText::FromString(TEXT("Busy"));
+			break;
+		}	
+
+		case ERemoteCacheState::Unavailable:
+		{
+			return FText::FromString(TEXT("Unavailable"));
+			break;
+		}
+
+		default:
+		case ERemoteCacheState::Warning:
+		{
+			return FText::FromString(TEXT("Warning"));
+			break;
+		}
+	}
+}
+
 FText SDerivedDataStatusBarWidget::GetRemoteCacheToolTipText() const
 {
 	FTextBuilder DescBuilder;
 
-	DescBuilder.AppendLineFormat(LOCTEXT("RemoteCacheConnectedText", "Remote Cache : {0}\n"), FText::FromString((FDerivedDataInformation::GetHasRemoteCache() ? TEXT("Connected") : TEXT("Unavailable"))));
-
+	DescBuilder.AppendLine(LOCTEXT("RemoteCacheToolTipText", "Remote Cache\n"));
+	DescBuilder.AppendLineFormat(LOCTEXT("RemoteCacheConnectedText", "Connected\t: {0}"), FText::FromString((FDerivedDataInformation::GetHasRemoteCache() ? TEXT("Yes") : TEXT("No"))));
+	DescBuilder.AppendLineFormat(LOCTEXT("RemoteCacheStatusText", "Status\t: {0}"), GetRemoteCacheStateAsText() );
+	
 	const double DownloadedBytesMB = FUnitConversion::Convert(FDerivedDataInformation::GetCacheActivitySizeBytes(true, false), EUnit::Bytes, EUnit::Megabytes);
 	const double UploadedBytesMB = FUnitConversion::Convert(FDerivedDataInformation::GetCacheActivitySizeBytes(false, false), EUnit::Bytes, EUnit::Megabytes);
 
