Fix for Zen Statistics in UI not adding up correctly

diff --git a/Engine/Source/Editor/DerivedDataEditor/Private/SVirtualAssetsStatistics.cpp b/Engine/Source/Editor/DerivedDataEditor/Private/SVirtualAssetsStatistics.cpp
--- a/Engine/Source/Editor/DerivedDataEditor/Private/SVirtualAssetsStatistics.cpp
+++ b/Engine/Source/Editor/DerivedDataEditor/Private/SVirtualAssetsStatistics.cpp
@@ -30,7 +30,7 @@ void SVirtualAssetsStatisticsDialog::Construct(const FArguments& InArgs)
 		SNew(SVerticalBox)
 		+ SVerticalBox::Slot()
 		.AutoHeight()
-		.Padding(0, 5, 0, 0)
+		.Padding(0, 20, 0, 0)
 		.Expose(GridSlot)
 		[
 			GetGridPanel()
@@ -68,7 +68,7 @@ TSharedRef<SWidget> SVirtualAssetsStatisticsDialog::GetGridPanel()
 	Panel->AddSlot(2, Row)
 	[
 		SNew(STextBlock)
-		.Margin(FMargin(ColumnMargin, RowMargin, 0.0f, TitleMargin))
+		.Margin(FMargin(ColumnMargin, RowMargin))
 		.ColorAndOpacity(TitleColor)
 		.Font(TitleFont)
 		.Justification(ETextJustify::Center)
@@ -78,7 +78,7 @@ TSharedRef<SWidget> SVirtualAssetsStatisticsDialog::GetGridPanel()
 	Panel->AddSlot(5, Row)
 	[
 		SNew(STextBlock)
-		.Margin(FMargin(ColumnMargin, RowMargin, 0.0f, TitleMargin))
+		.Margin(FMargin(ColumnMargin, RowMargin))
 		.ColorAndOpacity(TitleColor)
 		.Font(TitleFont)
 		.Justification(ETextJustify::Center)
@@ -88,7 +88,7 @@ TSharedRef<SWidget> SVirtualAssetsStatisticsDialog::GetGridPanel()
 	Panel->AddSlot(8, Row)
 	[
 		SNew(STextBlock)
-		.Margin(FMargin(ColumnMargin, RowMargin, 0.0f, TitleMargin))
+		.Margin(FMargin(ColumnMargin, RowMargin))
 		.ColorAndOpacity(TitleColor)
 		.Font(TitleFont)
 		.Justification(ETextJustify::Center)
@@ -199,12 +199,21 @@ TSharedRef<SWidget> SVirtualAssetsStatisticsDialog::GetGridPanel()
 
 	Row++;
 
-	auto GetVirtualAssetsStats = [&](const FString& BackendName, const FString& ConfigName, const FPayloadActivityInfo& AcitvityInfo)
+	IVirtualizationSystem& System = IVirtualizationSystem::Get();
+
+	FPayloadActivityInfo AccumulatedPayloadAcitvityInfo = System.GetAccumualtedPayloadActivityInfo();
+
+	FSlateColor Color = FStyleColors::Foreground;
+	FSlateFontInfo Font = FCoreStyle::GetDefaultFontStyle("Regular", 10);
+
+	auto DisplayPayloadActivityInfo = [&](const FString& BackendName, const FString& ConfigName, const FPayloadActivityInfo& PayloadAcitvityInfo)
 	{
 		Panel->AddSlot(0, Row)
 		[
 			SNew(STextBlock)
 			.Margin(FMargin(ColumnMargin, RowMargin))
+			.ColorAndOpacity(Color)
+			.Font(Font)
 			.Justification(ETextJustify::Left)
 			.Text(FText::FromString(ConfigName))
 		];
@@ -213,80 +222,101 @@ TSharedRef<SWidget> SVirtualAssetsStatisticsDialog::GetGridPanel()
 		[
 			SNew(STextBlock)
 			.Margin(FMargin(ColumnMargin, RowMargin))
+			.ColorAndOpacity(Color)
+			.Font(Font)
 			.Justification(ETextJustify::Center)
-			.Text_Lambda([AcitvityInfo] { return FText::FromString(FString::Printf(TEXT("%u"), AcitvityInfo.Pull.PayloadCount)); })
+			.Text_Lambda([PayloadAcitvityInfo] { return FText::FromString(FString::Printf(TEXT("%u"), PayloadAcitvityInfo.Pull.PayloadCount)); })
 		];
 
 		Panel->AddSlot(2, Row)
 		[
 			SNew(STextBlock)
 			.Margin(FMargin(ColumnMargin, RowMargin))
+			.ColorAndOpacity(Color)
+			.Font(Font)
 			.Justification(ETextJustify::Center)
-			.Text_Lambda([AcitvityInfo, BytesToMegaBytes] { return FText::FromString(SingleDecimalFormat((double)AcitvityInfo.Pull.TotalBytes * BytesToMegaBytes)); })
+			.Text_Lambda([PayloadAcitvityInfo, BytesToMegaBytes] { return FText::FromString(SingleDecimalFormat((double)PayloadAcitvityInfo.Pull.TotalBytes * BytesToMegaBytes)); })
 		];
 
 		Panel->AddSlot(3, Row)
 		[
 			SNew(STextBlock)
 			.Margin(FMargin(ColumnMargin, RowMargin))
+			.ColorAndOpacity(Color)
+			.Font(Font)
 			.Justification(ETextJustify::Center)
-			.Text_Lambda([AcitvityInfo] { return FText::FromString(SingleDecimalFormat((double)AcitvityInfo.Pull.CyclesSpent * FPlatformTime::GetSecondsPerCycle())); })
+			.Text_Lambda([PayloadAcitvityInfo] { return FText::FromString(SingleDecimalFormat((double)PayloadAcitvityInfo.Pull.CyclesSpent * FPlatformTime::GetSecondsPerCycle())); })
 		];
 
 		Panel->AddSlot(4, Row)
 		[
 			SNew(STextBlock)
 			.Margin(FMargin(ColumnMargin, RowMargin))
+			.ColorAndOpacity(Color)
+			.Font(Font)
 			.Justification(ETextJustify::Center)
-			.Text_Lambda([AcitvityInfo] { return FText::FromString(FString::Printf(TEXT("%u"), AcitvityInfo.Push.PayloadCount)); })
+			.Text_Lambda([PayloadAcitvityInfo] { return FText::FromString(FString::Printf(TEXT("%u"), PayloadAcitvityInfo.Push.PayloadCount)); })
 		];
 
 		Panel->AddSlot(5, Row)
 		[
 			SNew(STextBlock)
 			.Margin(FMargin(ColumnMargin, RowMargin))
+			.ColorAndOpacity(Color)
+			.Font(Font)
 			.Justification(ETextJustify::Center)
-			.Text_Lambda([AcitvityInfo, BytesToMegaBytes] { return FText::FromString(SingleDecimalFormat((double)AcitvityInfo.Push.TotalBytes * BytesToMegaBytes)); })
+			.Text_Lambda([PayloadAcitvityInfo, BytesToMegaBytes] { return FText::FromString(SingleDecimalFormat((double)PayloadAcitvityInfo.Push.TotalBytes * BytesToMegaBytes)); })
 		];
 
 		Panel->AddSlot(6, Row)
 		[
 			SNew(STextBlock)
 			.Margin(FMargin(ColumnMargin, RowMargin))
+			.ColorAndOpacity(Color)
+			.Font(Font)
 			.Justification(ETextJustify::Center)
-			.Text_Lambda([AcitvityInfo] { return FText::FromString(SingleDecimalFormat((double)AcitvityInfo.Push.CyclesSpent * FPlatformTime::GetSecondsPerCycle())); })
+			.Text_Lambda([PayloadAcitvityInfo] { return FText::FromString(SingleDecimalFormat((double)PayloadAcitvityInfo.Push.CyclesSpent * FPlatformTime::GetSecondsPerCycle())); })
 		];
 
 		Panel->AddSlot(7, Row)
 		[
 			SNew(STextBlock)
 			.Margin(FMargin(ColumnMargin, RowMargin))
+			.ColorAndOpacity(Color)
+			.Font(Font)
 			.Justification(ETextJustify::Center)
-			.Text_Lambda([AcitvityInfo] { return FText::FromString(FString::Printf(TEXT("%u"), AcitvityInfo.Cache.PayloadCount)); })
+			.Text_Lambda([PayloadAcitvityInfo] { return FText::FromString(FString::Printf(TEXT("%u"), PayloadAcitvityInfo.Cache.PayloadCount)); })
 		];
 
 		Panel->AddSlot(8, Row)
 		[
 			SNew(STextBlock)
 			.Margin(FMargin(ColumnMargin, RowMargin))
+			.ColorAndOpacity(Color)
+			.Font(Font)
 			.Justification(ETextJustify::Center)
-			.Text_Lambda([AcitvityInfo, BytesToMegaBytes] { return FText::FromString(SingleDecimalFormat((double)AcitvityInfo.Cache.TotalBytes * BytesToMegaBytes)); })
+			.Text_Lambda([PayloadAcitvityInfo, BytesToMegaBytes] { return FText::FromString(SingleDecimalFormat((double)PayloadAcitvityInfo.Cache.TotalBytes * BytesToMegaBytes)); })
 		];
 
 		Panel->AddSlot(9, Row)
 		[
 			SNew(STextBlock)
 			.Margin(FMargin(ColumnMargin, RowMargin))
+			.ColorAndOpacity(Color)
+			.Font(Font)
 			.Justification(ETextJustify::Center)
-			.Text_Lambda([AcitvityInfo] { return FText::FromString(SingleDecimalFormat((double)AcitvityInfo.Cache.CyclesSpent * FPlatformTime::GetSecondsPerCycle())); })
+			.Text_Lambda([PayloadAcitvityInfo] { return FText::FromString(SingleDecimalFormat((double)PayloadAcitvityInfo.Cache.CyclesSpent * FPlatformTime::GetSecondsPerCycle())); })
 		];
 
 		Row++;
 	};
 
-	IVirtualizationSystem& System = IVirtualizationSystem::Get();
+	System.GetPayloadActivityInfo(DisplayPayloadActivityInfo);
+
+	Color = TitleColor;
+	Font = TitleFont;
 
-	System.GetPayloadActivityInfo(GetVirtualAssetsStats);
+	DisplayPayloadActivityInfo(FString("Total"), FString("Total"), AccumulatedPayloadAcitvityInfo);
 
 	return Panel;
 }
diff --git a/Engine/Source/Editor/DerivedDataEditor/Private/SZenCacheStatistics.cpp b/Engine/Source/Editor/DerivedDataEditor/Private/SZenCacheStatistics.cpp
--- a/Engine/Source/Editor/DerivedDataEditor/Private/SZenCacheStatistics.cpp
+++ b/Engine/Source/Editor/DerivedDataEditor/Private/SZenCacheStatistics.cpp
@@ -14,6 +14,7 @@
 #define LOCTEXT_NAMESPACE "ZenEditor"
 
 extern FString SingleDecimalFormat(double Value);
+using namespace UE::Zen;
 
 void SZenCacheStatisticsDialog::Construct(const FArguments& InArgs)
 {
@@ -73,13 +74,22 @@ TSharedRef<SWidget> SZenCacheStatisticsDialog::GetGridPanel()
 
 #if UE_WITH_ZEN
 
-	UE::Zen::FZenStats ZenStats;
+	FZenStats ZenStats;
 
 	UE::Zen::GetDefaultServiceInstance().GetStats(ZenStats);
 
-	int32 Row = 0;
 	double SumTotalGetMB = 0.0;
 	double SumTotalPutMB = 0.0;
+	double TotalUpstreamHitRatio = 0.0;
+
+	for (const FZenEndPointStats& EndpointStats : ZenStats.UpstreamStats.EndPointStats)
+	{
+		SumTotalGetMB += EndpointStats.DownloadedMB;
+		SumTotalPutMB += EndpointStats.UploadedMB;
+		TotalUpstreamHitRatio += EndpointStats.HitRatio;
+	}
+
+	int32 Row = 0;
 
 	const float RowMargin = 0.0f;
 	const float TitleMargin = 10.0f;
@@ -166,7 +176,7 @@ TSharedRef<SWidget> SZenCacheStatisticsDialog::GetGridPanel()
 	[
 		SNew(STextBlock)
 		.Margin(FMargin(ColumnMargin, RowMargin))
-		.Text_Lambda([ZenStats] { return FText::FromString(SingleDecimalFormat(ZenStats.CacheStats.HitRatio * 100.0) + TEXT(" %")); })
+		.Text_Lambda([ZenStats, TotalUpstreamHitRatio] { return FText::FromString(SingleDecimalFormat( ( ZenStats.CacheStats.HitRatio- TotalUpstreamHitRatio ) * 100.0) + TEXT(" %")); })
 	];
 
 	Panel->AddSlot(5, Row)
@@ -180,7 +190,7 @@ TSharedRef<SWidget> SZenCacheStatisticsDialog::GetGridPanel()
 	
 	int32 EndpointIndex = 1;
 
-	for (const UE::Zen::FZenEndPointStats& EndpointStats : ZenStats.UpstreamStats.EndPointStats)
+	for (const FZenEndPointStats& EndpointStats : ZenStats.UpstreamStats.EndPointStats)
 	{
 		Panel->AddSlot(0, Row)
 		[
@@ -224,9 +234,6 @@ TSharedRef<SWidget> SZenCacheStatisticsDialog::GetGridPanel()
 			.Text_Lambda([EndpointStats] { return FText::FromString(EndpointStats.Name); })
 		];
 
-		SumTotalGetMB += EndpointStats.DownloadedMB;
-		SumTotalPutMB += EndpointStats.UploadedMB;
-
 		Row++;
 	}
 
@@ -240,6 +247,15 @@ TSharedRef<SWidget> SZenCacheStatisticsDialog::GetGridPanel()
 		.Justification(ETextJustify::Left)	
 	];
 
+	Panel->AddSlot(2, Row)
+	[
+		SNew(STextBlock)
+		.Margin(FMargin(ColumnMargin, RowMargin))
+		.ColorAndOpacity(TitleColor)
+		.Font(TitleFont)
+		.Text_Lambda([ZenStats] { return FText::FromString(SingleDecimalFormat( ZenStats.CacheStats.HitRatio * 100.0) + TEXT(" %")); })
+	];
+
 	Panel->AddSlot(3, Row)
 	.HAlign(HAlign_Right)
 	[
