Updated VA UI so that a clear message is given if the feature is currently disabled.

diff --git a/Engine/Source/Editor/DerivedDataEditor/Private/SVirtualAssetsStatistics.cpp b/Engine/Source/Editor/DerivedDataEditor/Private/SVirtualAssetsStatistics.cpp
--- a/Engine/Source/Editor/DerivedDataEditor/Private/SVirtualAssetsStatistics.cpp
+++ b/Engine/Source/Editor/DerivedDataEditor/Private/SVirtualAssetsStatistics.cpp
@@ -54,6 +54,10 @@ EActiveTimerReturnType SVirtualAssetsStatisticsDialog::UpdateGridPanels(double I
 
 TSharedRef<SWidget> SVirtualAssetsStatisticsDialog::GetGridPanel()
 {
+	IVirtualizationSystem& System = IVirtualizationSystem::Get();
+
+	TSharedRef<SGridPanel> Panel = SNew(SGridPanel);
+
 	const float RowMargin = 0.0f;
 	const float TitleMargin = 10.0f;
 	const float ColumnMargin = 10.0f;
@@ -61,262 +65,273 @@ TSharedRef<SWidget> SVirtualAssetsStatisticsDialog::GetGridPanel()
 	const FSlateFontInfo TitleFont = FCoreStyle::GetDefaultFontStyle("Bold", 10);
 	const double BytesToMegaBytes = 1.0 / (1024.0 * 1024.0);
 
-	TSharedRef<SGridPanel> Panel = SNew(SGridPanel);
-
-	int32 Row = 0;
-
-	Panel->AddSlot(2, Row)
-	[
-		SNew(STextBlock)
-		.Margin(FMargin(ColumnMargin, RowMargin))
-		.ColorAndOpacity(TitleColor)
-		.Font(TitleFont)
-		.Justification(ETextJustify::Center)
-		.Text(LOCTEXT("Read", "Read"))
-	];
-
-	Panel->AddSlot(5, Row)
-	[
-		SNew(STextBlock)
-		.Margin(FMargin(ColumnMargin, RowMargin))
-		.ColorAndOpacity(TitleColor)
-		.Font(TitleFont)
-		.Justification(ETextJustify::Center)
-		.Text(LOCTEXT("Write", "Write"))
-	];
-
-	Panel->AddSlot(8, Row)
-	[
-		SNew(STextBlock)
-		.Margin(FMargin(ColumnMargin, RowMargin))
-		.ColorAndOpacity(TitleColor)
-		.Font(TitleFont)
-		.Justification(ETextJustify::Center)
-		.Text(LOCTEXT("Cache", "Cache"))
-	];
-
-	Row++;
-
-	Panel->AddSlot(0, Row)
-	[
-		SNew(STextBlock)
-		.Margin(FMargin(ColumnMargin, RowMargin, 0.0f, TitleMargin))
-		.ColorAndOpacity(TitleColor)
-		.Font(TitleFont)
-		.Justification(ETextJustify::Left)
-		.Text(LOCTEXT("Backend", "Backend"))
-	];
-
-	Panel->AddSlot(1, Row)
-	[
-		SNew(STextBlock)
-		.Margin(FMargin(ColumnMargin, RowMargin, 0.0f, TitleMargin))
-		.ColorAndOpacity(TitleColor)
-		.Font(TitleFont)
-		.Justification(ETextJustify::Center)
-		.Text(LOCTEXT("Count", "Count"))
-	];
-
-	Panel->AddSlot(2, Row)
-	[
-		SNew(STextBlock)
-		.Margin(FMargin(ColumnMargin, RowMargin, 0.0f, TitleMargin))
-		.ColorAndOpacity(TitleColor)
-		.Font(TitleFont)
-		.Justification(ETextJustify::Center)
-		.Text(LOCTEXT("Time", "Time (Sec)"))
-	];
-
-	Panel->AddSlot(3, Row)
-	[
-		SNew(STextBlock)
-		.Margin(FMargin(ColumnMargin, RowMargin, 0.0f, TitleMargin))
-		.ColorAndOpacity(TitleColor)
-		.Font(TitleFont)
-		.Justification(ETextJustify::Center)
-		.Text(LOCTEXT("Size", "Size (MB)"))
-	];
-
-	Panel->AddSlot(4, Row)
-	[
-		SNew(STextBlock)
-		.Margin(FMargin(ColumnMargin, RowMargin, 0.0f, TitleMargin))
-		.ColorAndOpacity(TitleColor)
-		.Font(TitleFont)
-		.Justification(ETextJustify::Center)
-		.Text(LOCTEXT("Count", "Count"))
-	];
-
-	Panel->AddSlot(5, Row)
-	[
-		SNew(STextBlock)
-		.Margin(FMargin(ColumnMargin, RowMargin, 0.0f, TitleMargin))
-		.ColorAndOpacity(TitleColor)
-		.Font(TitleFont)
-		.Justification(ETextJustify::Center)
-		.Text(LOCTEXT("Time", "Time (Sec)"))
-	];
-
-	Panel->AddSlot(6, Row)
-	[
-		SNew(STextBlock)
-		.Margin(FMargin(ColumnMargin, RowMargin, 0.0f, TitleMargin))
-		.ColorAndOpacity(TitleColor)
-		.Font(TitleFont)
-		.Justification(ETextJustify::Center)
-		.Text(LOCTEXT("Size", "Size (MB)"))
-	];
-
-	Panel->AddSlot(7, Row)
-	[
-		SNew(STextBlock)
-		.Margin(FMargin(ColumnMargin, RowMargin, 0.0f, TitleMargin))
-		.ColorAndOpacity(TitleColor)
-		.Font(TitleFont)
-		.Justification(ETextJustify::Center)
-		.Text(LOCTEXT("Count", "Count"))
-	];
-
-	Panel->AddSlot(8, Row)
-	[
-		SNew(STextBlock)
-		.Margin(FMargin(ColumnMargin, RowMargin, 0.0f, TitleMargin))
-		.ColorAndOpacity(TitleColor)
-		.Font(TitleFont)
-		.Justification(ETextJustify::Center)
-		.Text(LOCTEXT("Time", "Time (Sec)"))
-	];
-
-	Panel->AddSlot(9, Row)
-	[
-		SNew(STextBlock)
-		.Margin(FMargin(ColumnMargin, RowMargin, 0.0f, TitleMargin))
-		.ColorAndOpacity(TitleColor)
-		.Font(TitleFont)
-		.Justification(ETextJustify::Center)
-		.Text(LOCTEXT("Size", "Size (MB)"))
-	];
+	if (System.IsEnabled() == false)
+	{
+		Panel->AddSlot(0,0)
+			[
+				SNew(STextBlock)
+				.Margin(FMargin(ColumnMargin, RowMargin))
+				.ColorAndOpacity(TitleColor)
+				.Font(TitleFont)
+				.Justification(ETextJustify::Center)
+				.Text(LOCTEXT("Disabled", "Virtual Assets Are Disabled For This Project"))
+			];
+	}
+	else
+	{
+		int32 Row = 0;
 
-	Row++;
+		Panel->AddSlot(2, Row)
+			[
+				SNew(STextBlock)
+				.Margin(FMargin(ColumnMargin, RowMargin))
+			.ColorAndOpacity(TitleColor)
+			.Font(TitleFont)
+			.Justification(ETextJustify::Center)
+			.Text(LOCTEXT("Read", "Read"))
+			];
 
-	IVirtualizationSystem& System = IVirtualizationSystem::Get();
+		Panel->AddSlot(5, Row)
+			[
+				SNew(STextBlock)
+				.Margin(FMargin(ColumnMargin, RowMargin))
+			.ColorAndOpacity(TitleColor)
+			.Font(TitleFont)
+			.Justification(ETextJustify::Center)
+			.Text(LOCTEXT("Write", "Write"))
+			];
 
-	FPayloadActivityInfo AccumulatedPayloadAcitvityInfo = System.GetAccumualtedPayloadActivityInfo();
+		Panel->AddSlot(8, Row)
+			[
+				SNew(STextBlock)
+				.Margin(FMargin(ColumnMargin, RowMargin))
+			.ColorAndOpacity(TitleColor)
+			.Font(TitleFont)
+			.Justification(ETextJustify::Center)
+			.Text(LOCTEXT("Cache", "Cache"))
+			];
 
-	FSlateColor Color = FStyleColors::Foreground;
-	FSlateFontInfo Font = FCoreStyle::GetDefaultFontStyle("Regular", 10);
+		Row++;
 
-	auto DisplayPayloadActivityInfo = [&](const FString& BackendName, const FString& ConfigName, const FPayloadActivityInfo& PayloadAcitvityInfo)
-	{
 		Panel->AddSlot(0, Row)
-		[
-			SNew(STextBlock)
-			.Margin(FMargin(ColumnMargin, RowMargin))
-			.ColorAndOpacity(Color)
-			.Font(Font)
+			[
+				SNew(STextBlock)
+				.Margin(FMargin(ColumnMargin, RowMargin, 0.0f, TitleMargin))
+			.ColorAndOpacity(TitleColor)
+			.Font(TitleFont)
 			.Justification(ETextJustify::Left)
-			.Text(FText::FromString(ConfigName))
-		];
+			.Text(LOCTEXT("Backend", "Backend"))
+			];
 
 		Panel->AddSlot(1, Row)
-		[
-			SNew(STextBlock)
-			.Margin(FMargin(ColumnMargin, RowMargin))
-			.ColorAndOpacity(Color)
-			.Font(Font)
+			[
+				SNew(STextBlock)
+				.Margin(FMargin(ColumnMargin, RowMargin, 0.0f, TitleMargin))
+			.ColorAndOpacity(TitleColor)
+			.Font(TitleFont)
 			.Justification(ETextJustify::Center)
-			.Text_Lambda([PayloadAcitvityInfo] { return FText::FromString(FString::Printf(TEXT("%u"), PayloadAcitvityInfo.Pull.PayloadCount)); })
-		];
+			.Text(LOCTEXT("Count", "Count"))
+			];
 
 		Panel->AddSlot(2, Row)
-		[
-			SNew(STextBlock)
-			.Margin(FMargin(ColumnMargin, RowMargin))
-			.ColorAndOpacity(Color)
-			.Font(Font)
+			[
+				SNew(STextBlock)
+				.Margin(FMargin(ColumnMargin, RowMargin, 0.0f, TitleMargin))
+			.ColorAndOpacity(TitleColor)
+			.Font(TitleFont)
 			.Justification(ETextJustify::Center)
-			.Text_Lambda([PayloadAcitvityInfo, BytesToMegaBytes] { return FText::FromString(SingleDecimalFormat((double)PayloadAcitvityInfo.Pull.TotalBytes * BytesToMegaBytes)); })
-		];
+			.Text(LOCTEXT("Time", "Time (Sec)"))
+			];
 
 		Panel->AddSlot(3, Row)
-		[
-			SNew(STextBlock)
-			.Margin(FMargin(ColumnMargin, RowMargin))
-			.ColorAndOpacity(Color)
-			.Font(Font)
+			[
+				SNew(STextBlock)
+				.Margin(FMargin(ColumnMargin, RowMargin, 0.0f, TitleMargin))
+			.ColorAndOpacity(TitleColor)
+			.Font(TitleFont)
 			.Justification(ETextJustify::Center)
-			.Text_Lambda([PayloadAcitvityInfo] { return FText::FromString(SingleDecimalFormat((double)PayloadAcitvityInfo.Pull.CyclesSpent * FPlatformTime::GetSecondsPerCycle())); })
-		];
+			.Text(LOCTEXT("Size", "Size (MB)"))
+			];
 
 		Panel->AddSlot(4, Row)
-		[
-			SNew(STextBlock)
-			.Margin(FMargin(ColumnMargin, RowMargin))
-			.ColorAndOpacity(Color)
-			.Font(Font)
+			[
+				SNew(STextBlock)
+				.Margin(FMargin(ColumnMargin, RowMargin, 0.0f, TitleMargin))
+			.ColorAndOpacity(TitleColor)
+			.Font(TitleFont)
 			.Justification(ETextJustify::Center)
-			.Text_Lambda([PayloadAcitvityInfo] { return FText::FromString(FString::Printf(TEXT("%u"), PayloadAcitvityInfo.Push.PayloadCount)); })
-		];
+			.Text(LOCTEXT("Count", "Count"))
+			];
 
 		Panel->AddSlot(5, Row)
-		[
-			SNew(STextBlock)
-			.Margin(FMargin(ColumnMargin, RowMargin))
-			.ColorAndOpacity(Color)
-			.Font(Font)
+			[
+				SNew(STextBlock)
+				.Margin(FMargin(ColumnMargin, RowMargin, 0.0f, TitleMargin))
+			.ColorAndOpacity(TitleColor)
+			.Font(TitleFont)
 			.Justification(ETextJustify::Center)
-			.Text_Lambda([PayloadAcitvityInfo, BytesToMegaBytes] { return FText::FromString(SingleDecimalFormat((double)PayloadAcitvityInfo.Push.TotalBytes * BytesToMegaBytes)); })
-		];
+			.Text(LOCTEXT("Time", "Time (Sec)"))
+			];
 
 		Panel->AddSlot(6, Row)
-		[
-			SNew(STextBlock)
-			.Margin(FMargin(ColumnMargin, RowMargin))
-			.ColorAndOpacity(Color)
-			.Font(Font)
+			[
+				SNew(STextBlock)
+				.Margin(FMargin(ColumnMargin, RowMargin, 0.0f, TitleMargin))
+			.ColorAndOpacity(TitleColor)
+			.Font(TitleFont)
 			.Justification(ETextJustify::Center)
-			.Text_Lambda([PayloadAcitvityInfo] { return FText::FromString(SingleDecimalFormat((double)PayloadAcitvityInfo.Push.CyclesSpent * FPlatformTime::GetSecondsPerCycle())); })
-		];
+			.Text(LOCTEXT("Size", "Size (MB)"))
+			];
 
 		Panel->AddSlot(7, Row)
-		[
-			SNew(STextBlock)
-			.Margin(FMargin(ColumnMargin, RowMargin))
-			.ColorAndOpacity(Color)
-			.Font(Font)
+			[
+				SNew(STextBlock)
+				.Margin(FMargin(ColumnMargin, RowMargin, 0.0f, TitleMargin))
+			.ColorAndOpacity(TitleColor)
+			.Font(TitleFont)
 			.Justification(ETextJustify::Center)
-			.Text_Lambda([PayloadAcitvityInfo] { return FText::FromString(FString::Printf(TEXT("%u"), PayloadAcitvityInfo.Cache.PayloadCount)); })
-		];
+			.Text(LOCTEXT("Count", "Count"))
+			];
 
 		Panel->AddSlot(8, Row)
-		[
-			SNew(STextBlock)
-			.Margin(FMargin(ColumnMargin, RowMargin))
-			.ColorAndOpacity(Color)
-			.Font(Font)
+			[
+				SNew(STextBlock)
+				.Margin(FMargin(ColumnMargin, RowMargin, 0.0f, TitleMargin))
+			.ColorAndOpacity(TitleColor)
+			.Font(TitleFont)
 			.Justification(ETextJustify::Center)
-			.Text_Lambda([PayloadAcitvityInfo, BytesToMegaBytes] { return FText::FromString(SingleDecimalFormat((double)PayloadAcitvityInfo.Cache.TotalBytes * BytesToMegaBytes)); })
-		];
+			.Text(LOCTEXT("Time", "Time (Sec)"))
+			];
 
 		Panel->AddSlot(9, Row)
-		[
-			SNew(STextBlock)
-			.Margin(FMargin(ColumnMargin, RowMargin))
-			.ColorAndOpacity(Color)
-			.Font(Font)
+			[
+				SNew(STextBlock)
+				.Margin(FMargin(ColumnMargin, RowMargin, 0.0f, TitleMargin))
+			.ColorAndOpacity(TitleColor)
+			.Font(TitleFont)
 			.Justification(ETextJustify::Center)
-			.Text_Lambda([PayloadAcitvityInfo] { return FText::FromString(SingleDecimalFormat((double)PayloadAcitvityInfo.Cache.CyclesSpent * FPlatformTime::GetSecondsPerCycle())); })
-		];
+			.Text(LOCTEXT("Size", "Size (MB)"))
+			];
 
 		Row++;
-	};
-
-	System.GetPayloadActivityInfo(DisplayPayloadActivityInfo);
-
-	Color = TitleColor;
-	Font = TitleFont;
 
-	DisplayPayloadActivityInfo(FString("Total"), FString("Total"), AccumulatedPayloadAcitvityInfo);
+		FPayloadActivityInfo AccumulatedPayloadAcitvityInfo = System.GetAccumualtedPayloadActivityInfo();
+
+		FSlateColor Color = FStyleColors::Foreground;
+		FSlateFontInfo Font = FCoreStyle::GetDefaultFontStyle("Regular", 10);
+
+		auto DisplayPayloadActivityInfo = [&](const FString& DebugName, const FString& ConfigName, const FPayloadActivityInfo& PayloadAcitvityInfo)
+		{
+			Panel->AddSlot(0, Row)
+				[
+					SNew(STextBlock)
+					.Margin(FMargin(ColumnMargin, RowMargin))
+				.ColorAndOpacity(Color)
+				.Font(Font)
+				.Justification(ETextJustify::Left)
+				.Text(FText::FromString(DebugName))
+				];
+
+			Panel->AddSlot(1, Row)
+				[
+					SNew(STextBlock)
+					.Margin(FMargin(ColumnMargin, RowMargin))
+				.ColorAndOpacity(Color)
+				.Font(Font)
+				.Justification(ETextJustify::Center)
+				.Text_Lambda([PayloadAcitvityInfo] { return FText::FromString(FString::Printf(TEXT("%u"), PayloadAcitvityInfo.Pull.PayloadCount)); })
+				];
+
+			Panel->AddSlot(2, Row)
+				[
+					SNew(STextBlock)
+					.Margin(FMargin(ColumnMargin, RowMargin))
+				.ColorAndOpacity(Color)
+				.Font(Font)
+				.Justification(ETextJustify::Center)
+				.Text_Lambda([PayloadAcitvityInfo, BytesToMegaBytes] { return FText::FromString(SingleDecimalFormat((double)PayloadAcitvityInfo.Pull.TotalBytes * BytesToMegaBytes)); })
+				];
+
+			Panel->AddSlot(3, Row)
+				[
+					SNew(STextBlock)
+					.Margin(FMargin(ColumnMargin, RowMargin))
+				.ColorAndOpacity(Color)
+				.Font(Font)
+				.Justification(ETextJustify::Center)
+				.Text_Lambda([PayloadAcitvityInfo] { return FText::FromString(SingleDecimalFormat((double)PayloadAcitvityInfo.Pull.CyclesSpent * FPlatformTime::GetSecondsPerCycle())); })
+				];
+
+			Panel->AddSlot(4, Row)
+				[
+					SNew(STextBlock)
+					.Margin(FMargin(ColumnMargin, RowMargin))
+				.ColorAndOpacity(Color)
+				.Font(Font)
+				.Justification(ETextJustify::Center)
+				.Text_Lambda([PayloadAcitvityInfo] { return FText::FromString(FString::Printf(TEXT("%u"), PayloadAcitvityInfo.Push.PayloadCount)); })
+				];
+
+			Panel->AddSlot(5, Row)
+				[
+					SNew(STextBlock)
+					.Margin(FMargin(ColumnMargin, RowMargin))
+				.ColorAndOpacity(Color)
+				.Font(Font)
+				.Justification(ETextJustify::Center)
+				.Text_Lambda([PayloadAcitvityInfo, BytesToMegaBytes] { return FText::FromString(SingleDecimalFormat((double)PayloadAcitvityInfo.Push.TotalBytes * BytesToMegaBytes)); })
+				];
+
+			Panel->AddSlot(6, Row)
+				[
+					SNew(STextBlock)
+					.Margin(FMargin(ColumnMargin, RowMargin))
+				.ColorAndOpacity(Color)
+				.Font(Font)
+				.Justification(ETextJustify::Center)
+				.Text_Lambda([PayloadAcitvityInfo] { return FText::FromString(SingleDecimalFormat((double)PayloadAcitvityInfo.Push.CyclesSpent * FPlatformTime::GetSecondsPerCycle())); })
+				];
+
+			Panel->AddSlot(7, Row)
+				[
+					SNew(STextBlock)
+					.Margin(FMargin(ColumnMargin, RowMargin))
+				.ColorAndOpacity(Color)
+				.Font(Font)
+				.Justification(ETextJustify::Center)
+				.Text_Lambda([PayloadAcitvityInfo] { return FText::FromString(FString::Printf(TEXT("%u"), PayloadAcitvityInfo.Cache.PayloadCount)); })
+				];
+
+			Panel->AddSlot(8, Row)
+				[
+					SNew(STextBlock)
+					.Margin(FMargin(ColumnMargin, RowMargin))
+				.ColorAndOpacity(Color)
+				.Font(Font)
+				.Justification(ETextJustify::Center)
+				.Text_Lambda([PayloadAcitvityInfo, BytesToMegaBytes] { return FText::FromString(SingleDecimalFormat((double)PayloadAcitvityInfo.Cache.TotalBytes * BytesToMegaBytes)); })
+				];
+
+			Panel->AddSlot(9, Row)
+				[
+					SNew(STextBlock)
+					.Margin(FMargin(ColumnMargin, RowMargin))
+				.ColorAndOpacity(Color)
+				.Font(Font)
+				.Justification(ETextJustify::Center)
+				.Text_Lambda([PayloadAcitvityInfo] { return FText::FromString(SingleDecimalFormat((double)PayloadAcitvityInfo.Cache.CyclesSpent * FPlatformTime::GetSecondsPerCycle())); })
+				];
+
+			Row++;
+		};
+
+		System.GetPayloadActivityInfo(DisplayPayloadActivityInfo);
+
+		Color = TitleColor;
+		Font = TitleFont;
+
+		DisplayPayloadActivityInfo(FString("Total"), FString("Total"), AccumulatedPayloadAcitvityInfo);
+	}
 
 	return Panel;
 }
