Added Notification events to VA system
Added notification UI for reporting errors syncing bulk data
Added notification UI for displaying sync activity

diff --git a/Engine/Source/Developer/Virtualization/Private/VirtualizationManager.cpp b/Engine/Source/Developer/Virtualization/Private/VirtualizationManager.cpp
--- a/Engine/Source/Developer/Virtualization/Private/VirtualizationManager.cpp
+++ b/Engine/Source/Developer/Virtualization/Private/VirtualizationManager.cpp
@@ -392,20 +392,31 @@ FCompressedBuffer FVirtualizationManager::PullData(const FPayloadId& Id)
 
 	FConditionalScopeLock _(&ForceSingleThreadedCS, bForceSingleThreaded);
 
+	GetNotificationEvent().Broadcast(IVirtualizationSystem::PullBegunNotification, Id);
+
 	for (IVirtualizationBackend* Backend : PullEnabledBackends)
 	{
 		FCompressedBuffer Payload = PullDataFromBackend(*Backend, Id);
+
 		if (Payload)
 		{
 			if (bEnableCacheAfterPull)
 			{
 				CachePayload(Id, Payload, Backend);
 			}
-			
+
+			GetNotificationEvent().Broadcast(IVirtualizationSystem::PullEndedNotification, Id);
+
 			return Payload;
 		}
 	}
 
+	// Pull ended but failed..
+	GetNotificationEvent().Broadcast(IVirtualizationSystem::PullEndedNotification, Id);
+
+	// Broadcast the pull failed event to any listeners
+	GetNotificationEvent().Broadcast(IVirtualizationSystem::PullFailedNotification, Id);
+
 	// TODO: Maybe this should be a fatal error? If we keep it as an error we need to make sure any calling
 	// code handles it properly.
 	// Could be worth extending ::PullData to return error codes instead so we can make a better distinction 
diff --git a/Engine/Source/Developer/Virtualization/Private/VirtualizationManager.h b/Engine/Source/Developer/Virtualization/Private/VirtualizationManager.h
--- a/Engine/Source/Developer/Virtualization/Private/VirtualizationManager.h
+++ b/Engine/Source/Developer/Virtualization/Private/VirtualizationManager.h
@@ -123,6 +123,13 @@ public:
 
 	/** Access profiling info relating to payload activity per backend. Stats will only be collected if ENABLE_COOK_STATS is enabled.*/
 	virtual void GetPayloadActivityInfo( GetPayloadActivityInfoFuncRef ) const override;
+
+
+	/** Get event delegate for system notifications */
+	virtual FOnNotification& GetNotificationEvent() override
+	{
+		return NotificationEvent;
+	}
 	
 private:
 	
@@ -203,6 +210,9 @@ private:
 	 * and can contain a mixture of local cacheable and persistent backends 
 	 */
 	FBackendArray PullEnabledBackends;
+
+	/** Our notification Event */
+	FOnNotification NotificationEvent;
 };
 
 } // namespace UE::Virtualization
diff --git a/Engine/Source/Editor/DerivedDataEditor/Private/SVirtualAssetsStatistics.cpp b/Engine/Source/Editor/DerivedDataEditor/Private/SVirtualAssetsStatistics.cpp
--- a/Engine/Source/Editor/DerivedDataEditor/Private/SVirtualAssetsStatistics.cpp
+++ b/Engine/Source/Editor/DerivedDataEditor/Private/SVirtualAssetsStatistics.cpp
@@ -10,6 +10,7 @@
 #include "Widgets/Text/STextBlock.h"
 #include "Widgets/SBoxPanel.h"
 #include "Internationalization/FastDecimalFormat.h"
+#include "Framework/Notifications/NotificationManager.h"
 
 #define LOCTEXT_NAMESPACE "VirtualAssets"
 
@@ -17,14 +18,60 @@ extern FString SingleDecimalFormat(double Value);
 
 using namespace UE::Virtualization;
 
-void SVirtualAssetsStatisticsDialog::Construct(const FArguments& InArgs)
+SVirtualAssetsStatisticsDialog::SVirtualAssetsStatisticsDialog()
 {
-	const float RowMargin = 0.0f;
-	const float TitleMargin = 10.0f;
-	const float ColumnMargin = 10.0f;
-	const FSlateColor TitleColour = FStyleColors::AccentWhite;
-	const FSlateFontInfo TitleFont = FCoreStyle::GetDefaultFontStyle("Bold", 10);
+	// Register our VA notification delegate with the event
+	IVirtualizationSystem& System = IVirtualizationSystem::Get();
+	System.GetNotificationEvent().AddRaw(this, &SVirtualAssetsStatisticsDialog::OnNotificationEvent);
+
+}
+
+SVirtualAssetsStatisticsDialog::~SVirtualAssetsStatisticsDialog()
+{
+	// Unregister our VA notification delegate with the event
+	IVirtualizationSystem& System = IVirtualizationSystem::Get();
+	System.GetNotificationEvent().RemoveAll(this);
+}
+
+void SVirtualAssetsStatisticsDialog::OnNotificationEvent(IVirtualizationSystem::ENotification Notification, const FPayloadId& PayloadId)
+{
+	FScopeLock SocpeLock(&NotificationCS);
+	
+	switch (Notification)
+	{	
+		case IVirtualizationSystem::ENotification::PullBegunNotification:
+		{
+			NumPullRequests++;
+			break;
+		}
+
+		case IVirtualizationSystem::ENotification::PullEndedNotification:
+		{	
+			NumPullRequests--;
+			break;
+		}
+
+		case IVirtualizationSystem::ENotification::PullFailedNotification:
+		{
+			FNotificationInfo* Info = new FNotificationInfo(LOCTEXT("PayloadSyncFail", "Failed To Sync Asset Payload"));
 
+			Info->bUseSuccessFailIcons = true;
+			Info->bFireAndForget = true;
+			Info->FadeOutDuration = 1.0f;
+			Info->ExpireDuration = 4.0f;
+
+			FSlateNotificationManager::Get().QueueNotification(Info);
+
+			break;
+		}
+
+		default:
+		break;
+	}
+}
+
+void SVirtualAssetsStatisticsDialog::Construct(const FArguments& InArgs)
+{
 	this->ChildSlot
 	[
 		SNew(SVerticalBox)
@@ -49,6 +96,32 @@ EActiveTimerReturnType SVirtualAssetsStatisticsDialog::UpdateGridPanels(double I
 
 	SlatePrepass(GetPrepassLayoutScaleMultiplier());
 
+	if ( NumPullRequests!=0 && PullRequestNotificationItem.IsValid()==false )
+	{
+		// No existing notification or the existing one has finished
+		TPromise<TWeakPtr<SNotificationItem>> NotificationPromise;
+
+		FNotificationInfo Info(LOCTEXT("PayloadSyncInProgress", "Syncing Asset Payloads"));
+		Info.bFireAndForget = false;
+		Info.FadeOutDuration = 0.5f;
+		Info.ExpireDuration = 0.0f;
+
+		PullRequestNotificationItem = FSlateNotificationManager::Get().AddNotification(Info);
+
+		if (PullRequestNotificationItem.IsValid())
+		{
+			NotificationPromise.SetValue(PullRequestNotificationItem);
+			PullRequestNotificationItem->SetCompletionState(SNotificationItem::CS_Pending);
+		}
+	}
+	
+	if ( NumPullRequests==0 && PullRequestNotificationItem.IsValid()==true )
+	{
+		PullRequestNotificationItem->SetCompletionState(SNotificationItem::CS_Success);
+		PullRequestNotificationItem->ExpireAndFadeout();
+		PullRequestNotificationItem.Reset();
+	}
+
 	return EActiveTimerReturnType::Continue;
 }
 
diff --git a/Engine/Source/Editor/DerivedDataEditor/Private/SVirtualAssetsStatistics.h b/Engine/Source/Editor/DerivedDataEditor/Private/SVirtualAssetsStatistics.h
--- a/Engine/Source/Editor/DerivedDataEditor/Private/SVirtualAssetsStatistics.h
+++ b/Engine/Source/Editor/DerivedDataEditor/Private/SVirtualAssetsStatistics.h
@@ -7,14 +7,18 @@
 #include "Widgets/SWidget.h"
 #include "Widgets/SCompoundWidget.h"
 #include "Widgets/SBoxPanel.h"
-#include "ZenServerInterface.h"
+#include "Widgets/Notifications/SNotificationList.h"
+#include "VirtualizationManager.h"
 
 class SVirtualAssetsStatisticsDialog : public SCompoundWidget
 {
 	SLATE_BEGIN_ARGS(SVirtualAssetsStatisticsDialog) {}
 	SLATE_END_ARGS()
 
-		void Construct(const FArguments& InArgs);
+	void Construct(const FArguments& InArgs);
+
+	SVirtualAssetsStatisticsDialog();
+	virtual ~SVirtualAssetsStatisticsDialog();
 
 private:
 
@@ -22,5 +26,15 @@ private:
 
 	EActiveTimerReturnType UpdateGridPanels(double InCurrentTime, float InDeltaTime);
 
+	FText GetNotificationText() const;
+
+	void OnNotificationEvent(UE::Virtualization::IVirtualizationSystem::ENotification Notification, const UE::Virtualization::FPayloadId& PayloadId);
+
 	SVerticalBox::FSlot* GridSlot = nullptr;
+
+	FCriticalSection NotificationCS;
+
+	TSharedPtr<SNotificationItem> PullRequestNotificationItem;
+
+	uint32 NumPullRequests = 0;
 };
\ No newline at end of file
diff --git a/Engine/Source/Runtime/Core/Private/Virtualization/VirtualizationSystem.cpp b/Engine/Source/Runtime/Core/Private/Virtualization/VirtualizationSystem.cpp
--- a/Engine/Source/Runtime/Core/Private/Virtualization/VirtualizationSystem.cpp
+++ b/Engine/Source/Runtime/Core/Private/Virtualization/VirtualizationSystem.cpp
@@ -42,6 +42,13 @@ public:
 	{
 		return FPayloadActivityInfo();
 	}
+
+	virtual FOnNotification& GetNotificationEvent() override
+	{
+		return NotificationEvent;
+	}
+
+	FOnNotification NotificationEvent;
 };
 
 TUniquePtr<IVirtualizationSystem> GVirtualizationSystem = nullptr;
diff --git a/Engine/Source/Runtime/Core/Public/Virtualization/VirtualizationSystem.h b/Engine/Source/Runtime/Core/Public/Virtualization/VirtualizationSystem.h
--- a/Engine/Source/Runtime/Core/Public/Virtualization/VirtualizationSystem.h
+++ b/Engine/Source/Runtime/Core/Public/Virtualization/VirtualizationSystem.h
@@ -99,13 +99,30 @@ public:
 	 */
 	virtual FCompressedBuffer PullData(const FPayloadId& Id) = 0;
 
-	using GetPayloadActivityInfoFuncRef = TFunctionRef<void(const FString& DebugName, const FString& ConfigName, const FPayloadActivityInfo& PAyloadInfo)>;
+	using GetPayloadActivityInfoFuncRef = TFunctionRef<void(const FString& DebugName, const FString& ConfigName, const FPayloadActivityInfo& PayloadInfo)>;
 
 	/** Access profiling info relating to payload activity per backend. Stats will only be collected if ENABLE_COOK_STATS is enabled.*/
 	virtual void GetPayloadActivityInfo( GetPayloadActivityInfoFuncRef ) const = 0;
 
 	/** Access profiling info relating to accumulated payload activity. Stats will only be collected if ENABLE_COOK_STATS is enabled.*/
 	virtual FPayloadActivityInfo GetAccumualtedPayloadActivityInfo() const = 0;
+
+	//* Notification messages
+	enum ENotification
+	{
+		PushBegunNotification,
+		PushEndedNotification,
+		PushFailedNotification,
+
+		PullBegunNotification,
+		PullEndedNotification,
+		PullFailedNotification,
+	};
+
+	/** Declare delegate for notifications*/
+	DECLARE_MULTICAST_DELEGATE_TwoParams(FOnNotification, ENotification, const FPayloadId&);
+
+	virtual FOnNotification& GetNotificationEvent() = 0;
 };
 
 namespace Private
