Fixed issue where the Virtual Payload PullCount became negative

diff --git a/Engine/Source/Editor/DerivedDataEditor/Private/SVirtualAssetsStatistics.cpp b/Engine/Source/Editor/DerivedDataEditor/Private/SVirtualAssetsStatistics.cpp
--- a/Engine/Source/Editor/DerivedDataEditor/Private/SVirtualAssetsStatistics.cpp
+++ b/Engine/Source/Editor/DerivedDataEditor/Private/SVirtualAssetsStatistics.cpp
@@ -41,13 +41,20 @@ void SVirtualAssetsStatisticsDialog::OnNotificationEvent(IVirtualizationSystem::
 	{	
 		case IVirtualizationSystem::ENotification::PullBegunNotification:
 		{
+			IsPulling = true;
 			NumPullRequests++;
+
 			break;
 		}
 
 		case IVirtualizationSystem::ENotification::PullEndedNotification:
 		{	
-			NumPullRequests--;
+			if (IsPulling == true)
+			{
+				NumPullRequests--;
+				IsPulling = NumPullRequests!=0;
+			}
+			
 			break;
 		}
 
diff --git a/Engine/Source/Editor/DerivedDataEditor/Private/SVirtualAssetsStatistics.h b/Engine/Source/Editor/DerivedDataEditor/Private/SVirtualAssetsStatistics.h
--- a/Engine/Source/Editor/DerivedDataEditor/Private/SVirtualAssetsStatistics.h
+++ b/Engine/Source/Editor/DerivedDataEditor/Private/SVirtualAssetsStatistics.h
@@ -36,5 +36,6 @@ private:
 
 	TSharedPtr<SNotificationItem> PullRequestNotificationItem;
 
-	uint32 NumPullRequests = 0;
+	bool	IsPulling = false;
+	uint32	NumPullRequests = 0;
 };
\ No newline at end of file
