Put the Zen cache hits ratio display back to how it was before.

diff --git a/Engine/Source/Editor/DerivedDataEditor/Private/SZenCacheStatistics.cpp b/Engine/Source/Editor/DerivedDataEditor/Private/SZenCacheStatistics.cpp
--- a/Engine/Source/Editor/DerivedDataEditor/Private/SZenCacheStatistics.cpp
+++ b/Engine/Source/Editor/DerivedDataEditor/Private/SZenCacheStatistics.cpp
@@ -176,7 +176,7 @@ TSharedRef<SWidget> SZenCacheStatisticsDialog::GetGridPanel()
 	[
 		SNew(STextBlock)
 		.Margin(FMargin(ColumnMargin, RowMargin))
-		.Text_Lambda([ZenStats, TotalUpstreamHitRatio] { return FText::FromString(SingleDecimalFormat( ( ZenStats.CacheStats.HitRatio- TotalUpstreamHitRatio ) * 100.0) + TEXT(" %")); })
+		.Text_Lambda([ZenStats, TotalUpstreamHitRatio] { return FText::FromString(SingleDecimalFormat( ZenStats.CacheStats.HitRatio * 100.0) + TEXT(" %")); })
 	];
 
 	Panel->AddSlot(5, Row)
