The virtualized asset statistic panel now displays the correct values in the Time and Size columns

#rnx

- The time column was showing the size data and vice versa, as the columns were originally the other way around.
- Fixed alignment issues
- Fixed a typo in the DisplayPayloadActivityInfo lambda

diff --git a/Engine/Source/Editor/DerivedDataEditor/Private/SVirtualAssetsStatistics.cpp b/Engine/Source/Editor/DerivedDataEditor/Private/SVirtualAssetsStatistics.cpp
--- a/Engine/Source/Editor/DerivedDataEditor/Private/SVirtualAssetsStatistics.cpp
+++ b/Engine/Source/Editor/DerivedDataEditor/Private/SVirtualAssetsStatistics.cpp
@@ -165,30 +165,30 @@ TSharedRef<SWidget> SVirtualAssetsStatisticsDialog::GetGridPanel()
 			[
 				SNew(STextBlock)
 				.Margin(FMargin(ColumnMargin, RowMargin))
-			.ColorAndOpacity(TitleColor)
-			.Font(TitleFont)
-			.Justification(ETextJustify::Center)
-			.Text(LOCTEXT("Read", "Read"))
+				.ColorAndOpacity(TitleColor)
+				.Font(TitleFont)
+				.Justification(ETextJustify::Center)
+				.Text(LOCTEXT("Read", "Read"))
 			];
 
 		Panel->AddSlot(5, Row)
 			[
 				SNew(STextBlock)
 				.Margin(FMargin(ColumnMargin, RowMargin))
-			.ColorAndOpacity(TitleColor)
-			.Font(TitleFont)
-			.Justification(ETextJustify::Center)
-			.Text(LOCTEXT("Write", "Write"))
+				.ColorAndOpacity(TitleColor)
+				.Font(TitleFont)
+				.Justification(ETextJustify::Center)
+				.Text(LOCTEXT("Write", "Write"))
 			];
 
 		Panel->AddSlot(8, Row)
 			[
 				SNew(STextBlock)
 				.Margin(FMargin(ColumnMargin, RowMargin))
-			.ColorAndOpacity(TitleColor)
-			.Font(TitleFont)
-			.Justification(ETextJustify::Center)
-			.Text(LOCTEXT("Cache", "Cache"))
+				.ColorAndOpacity(TitleColor)
+				.Font(TitleFont)
+				.Justification(ETextJustify::Center)
+				.Text(LOCTEXT("Cache", "Cache"))
 			];
 
 		Row++;
@@ -197,100 +197,100 @@ TSharedRef<SWidget> SVirtualAssetsStatisticsDialog::GetGridPanel()
 			[
 				SNew(STextBlock)
 				.Margin(FMargin(ColumnMargin, RowMargin, 0.0f, TitleMargin))
-			.ColorAndOpacity(TitleColor)
-			.Font(TitleFont)
-			.Justification(ETextJustify::Left)
-			.Text(LOCTEXT("Backend", "Backend"))
+				.ColorAndOpacity(TitleColor)
+				.Font(TitleFont)
+				.Justification(ETextJustify::Left)
+				.Text(LOCTEXT("Backend", "Backend"))
 			];
 
 		Panel->AddSlot(1, Row)
 			[
 				SNew(STextBlock)
 				.Margin(FMargin(ColumnMargin, RowMargin, 0.0f, TitleMargin))
-			.ColorAndOpacity(TitleColor)
-			.Font(TitleFont)
-			.Justification(ETextJustify::Center)
-			.Text(LOCTEXT("Count", "Count"))
+				.ColorAndOpacity(TitleColor)
+				.Font(TitleFont)
+				.Justification(ETextJustify::Center)
+				.Text(LOCTEXT("Count", "Count"))
 			];
 
 		Panel->AddSlot(2, Row)
 			[
 				SNew(STextBlock)
 				.Margin(FMargin(ColumnMargin, RowMargin, 0.0f, TitleMargin))
-			.ColorAndOpacity(TitleColor)
-			.Font(TitleFont)
-			.Justification(ETextJustify::Center)
-			.Text(LOCTEXT("Time", "Time (Sec)"))
+				.ColorAndOpacity(TitleColor)
+				.Font(TitleFont)
+				.Justification(ETextJustify::Center)
+				.Text(LOCTEXT("Time", "Time (Sec)"))
 			];
 
 		Panel->AddSlot(3, Row)
 			[
 				SNew(STextBlock)
 				.Margin(FMargin(ColumnMargin, RowMargin, 0.0f, TitleMargin))
-			.ColorAndOpacity(TitleColor)
-			.Font(TitleFont)
-			.Justification(ETextJustify::Center)
-			.Text(LOCTEXT("Size", "Size (MB)"))
+				.ColorAndOpacity(TitleColor)
+				.Font(TitleFont)
+				.Justification(ETextJustify::Center)
+				.Text(LOCTEXT("Size", "Size (MB)"))
 			];
 
 		Panel->AddSlot(4, Row)
 			[
 				SNew(STextBlock)
 				.Margin(FMargin(ColumnMargin, RowMargin, 0.0f, TitleMargin))
-			.ColorAndOpacity(TitleColor)
-			.Font(TitleFont)
-			.Justification(ETextJustify::Center)
-			.Text(LOCTEXT("Count", "Count"))
+				.ColorAndOpacity(TitleColor)
+				.Font(TitleFont)
+				.Justification(ETextJustify::Center)
+				.Text(LOCTEXT("Count", "Count"))
 			];
 
 		Panel->AddSlot(5, Row)
 			[
 				SNew(STextBlock)
 				.Margin(FMargin(ColumnMargin, RowMargin, 0.0f, TitleMargin))
-			.ColorAndOpacity(TitleColor)
-			.Font(TitleFont)
-			.Justification(ETextJustify::Center)
-			.Text(LOCTEXT("Time", "Time (Sec)"))
+				.ColorAndOpacity(TitleColor)
+				.Font(TitleFont)
+				.Justification(ETextJustify::Center)
+				.Text(LOCTEXT("Time", "Time (Sec)"))
 			];
 
 		Panel->AddSlot(6, Row)
 			[
 				SNew(STextBlock)
 				.Margin(FMargin(ColumnMargin, RowMargin, 0.0f, TitleMargin))
-			.ColorAndOpacity(TitleColor)
-			.Font(TitleFont)
-			.Justification(ETextJustify::Center)
-			.Text(LOCTEXT("Size", "Size (MB)"))
+				.ColorAndOpacity(TitleColor)
+				.Font(TitleFont)
+				.Justification(ETextJustify::Center)
+				.Text(LOCTEXT("Size", "Size (MB)"))
 			];
 
 		Panel->AddSlot(7, Row)
 			[
 				SNew(STextBlock)
 				.Margin(FMargin(ColumnMargin, RowMargin, 0.0f, TitleMargin))
-			.ColorAndOpacity(TitleColor)
-			.Font(TitleFont)
-			.Justification(ETextJustify::Center)
-			.Text(LOCTEXT("Count", "Count"))
+				.ColorAndOpacity(TitleColor)
+				.Font(TitleFont)
+				.Justification(ETextJustify::Center)
+				.Text(LOCTEXT("Count", "Count"))
 			];
 
 		Panel->AddSlot(8, Row)
 			[
 				SNew(STextBlock)
 				.Margin(FMargin(ColumnMargin, RowMargin, 0.0f, TitleMargin))
-			.ColorAndOpacity(TitleColor)
-			.Font(TitleFont)
-			.Justification(ETextJustify::Center)
-			.Text(LOCTEXT("Time", "Time (Sec)"))
+				.ColorAndOpacity(TitleColor)
+				.Font(TitleFont)
+				.Justification(ETextJustify::Center)
+				.Text(LOCTEXT("Time", "Time (Sec)"))
 			];
 
 		Panel->AddSlot(9, Row)
 			[
 				SNew(STextBlock)
 				.Margin(FMargin(ColumnMargin, RowMargin, 0.0f, TitleMargin))
-			.ColorAndOpacity(TitleColor)
-			.Font(TitleFont)
-			.Justification(ETextJustify::Center)
-			.Text(LOCTEXT("Size", "Size (MB)"))
+				.ColorAndOpacity(TitleColor)
+				.Font(TitleFont)
+				.Justification(ETextJustify::Center)
+				.Text(LOCTEXT("Size", "Size (MB)"))
 			];
 
 		Row++;
@@ -300,106 +300,106 @@ TSharedRef<SWidget> SVirtualAssetsStatisticsDialog::GetGridPanel()
 		FSlateColor Color = FStyleColors::Foreground;
 		FSlateFontInfo Font = FCoreStyle::GetDefaultFontStyle("Regular", 10);
 
-		auto DisplayPayloadActivityInfo = [&](const FString& DebugName, const FString& ConfigName, const FPayloadActivityInfo& PayloadAcitvityInfo)
+		auto DisplayPayloadActivityInfo = [&](const FString& DebugName, const FString& ConfigName, const FPayloadActivityInfo& PayloadActivityInfo)
 		{
 			Panel->AddSlot(0, Row)
 				[
 					SNew(STextBlock)
 					.Margin(FMargin(ColumnMargin, RowMargin))
-				.ColorAndOpacity(Color)
-				.Font(Font)
-				.Justification(ETextJustify::Left)
-				.Text(FText::FromString(DebugName))
+					.ColorAndOpacity(Color)
+					.Font(Font)
+					.Justification(ETextJustify::Left)
+					.Text(FText::FromString(DebugName))
 				];
 
 			Panel->AddSlot(1, Row)
 				[
 					SNew(STextBlock)
 					.Margin(FMargin(ColumnMargin, RowMargin))
-				.ColorAndOpacity(Color)
-				.Font(Font)
-				.Justification(ETextJustify::Center)
-				.Text_Lambda([PayloadAcitvityInfo] { return FText::FromString(FString::Printf(TEXT("%u"), PayloadAcitvityInfo.Pull.PayloadCount)); })
+					.ColorAndOpacity(Color)
+					.Font(Font)
+					.Justification(ETextJustify::Center)
+					.Text_Lambda([PayloadActivityInfo] { return FText::FromString(FString::Printf(TEXT("%u"), PayloadActivityInfo.Pull.PayloadCount)); })
 				];
 
 			Panel->AddSlot(2, Row)
 				[
 					SNew(STextBlock)
 					.Margin(FMargin(ColumnMargin, RowMargin))
-				.ColorAndOpacity(Color)
-				.Font(Font)
-				.Justification(ETextJustify::Center)
-				.Text_Lambda([PayloadAcitvityInfo, BytesToMegaBytes] { return FText::FromString(SingleDecimalFormat((double)PayloadAcitvityInfo.Pull.TotalBytes * BytesToMegaBytes)); })
+					.ColorAndOpacity(Color)
+					.Font(Font)
+					.Justification(ETextJustify::Center)
+					.Text_Lambda([PayloadActivityInfo] { return FText::FromString(SingleDecimalFormat((double)PayloadActivityInfo.Pull.CyclesSpent * FPlatformTime::GetSecondsPerCycle())); })
 				];
 
 			Panel->AddSlot(3, Row)
 				[
 					SNew(STextBlock)
 					.Margin(FMargin(ColumnMargin, RowMargin))
-				.ColorAndOpacity(Color)
-				.Font(Font)
-				.Justification(ETextJustify::Center)
-				.Text_Lambda([PayloadAcitvityInfo] { return FText::FromString(SingleDecimalFormat((double)PayloadAcitvityInfo.Pull.CyclesSpent * FPlatformTime::GetSecondsPerCycle())); })
+					.ColorAndOpacity(Color)
+					.Font(Font)
+					.Justification(ETextJustify::Center)
+					.Text_Lambda([PayloadActivityInfo, BytesToMegaBytes] { return FText::FromString(SingleDecimalFormat((double)PayloadActivityInfo.Pull.TotalBytes * BytesToMegaBytes)); })
 				];
 
 			Panel->AddSlot(4, Row)
 				[
 					SNew(STextBlock)
 					.Margin(FMargin(ColumnMargin, RowMargin))
-				.ColorAndOpacity(Color)
-				.Font(Font)
-				.Justification(ETextJustify::Center)
-				.Text_Lambda([PayloadAcitvityInfo] { return FText::FromString(FString::Printf(TEXT("%u"), PayloadAcitvityInfo.Push.PayloadCount)); })
+					.ColorAndOpacity(Color)
+					.Font(Font)
+					.Justification(ETextJustify::Center)
+					.Text_Lambda([PayloadActivityInfo] { return FText::FromString(FString::Printf(TEXT("%u"), PayloadActivityInfo.Push.PayloadCount)); })
 				];
 
 			Panel->AddSlot(5, Row)
 				[
 					SNew(STextBlock)
 					.Margin(FMargin(ColumnMargin, RowMargin))
-				.ColorAndOpacity(Color)
-				.Font(Font)
-				.Justification(ETextJustify::Center)
-				.Text_Lambda([PayloadAcitvityInfo, BytesToMegaBytes] { return FText::FromString(SingleDecimalFormat((double)PayloadAcitvityInfo.Push.TotalBytes * BytesToMegaBytes)); })
+					.ColorAndOpacity(Color)
+					.Font(Font)
+					.Justification(ETextJustify::Center)
+					.Text_Lambda([PayloadActivityInfo] { return FText::FromString(SingleDecimalFormat((double)PayloadActivityInfo.Push.CyclesSpent * FPlatformTime::GetSecondsPerCycle())); })
 				];
 
 			Panel->AddSlot(6, Row)
 				[
 					SNew(STextBlock)
 					.Margin(FMargin(ColumnMargin, RowMargin))
-				.ColorAndOpacity(Color)
-				.Font(Font)
-				.Justification(ETextJustify::Center)
-				.Text_Lambda([PayloadAcitvityInfo] { return FText::FromString(SingleDecimalFormat((double)PayloadAcitvityInfo.Push.CyclesSpent * FPlatformTime::GetSecondsPerCycle())); })
+					.ColorAndOpacity(Color)
+					.Font(Font)
+					.Justification(ETextJustify::Center)
+					.Text_Lambda([PayloadActivityInfo, BytesToMegaBytes] { return FText::FromString(SingleDecimalFormat((double)PayloadActivityInfo.Push.TotalBytes * BytesToMegaBytes)); })
 				];
 
 			Panel->AddSlot(7, Row)
 				[
 					SNew(STextBlock)
 					.Margin(FMargin(ColumnMargin, RowMargin))
-				.ColorAndOpacity(Color)
-				.Font(Font)
-				.Justification(ETextJustify::Center)
-				.Text_Lambda([PayloadAcitvityInfo] { return FText::FromString(FString::Printf(TEXT("%u"), PayloadAcitvityInfo.Cache.PayloadCount)); })
+					.ColorAndOpacity(Color)
+					.Font(Font)
+					.Justification(ETextJustify::Center)
+					.Text_Lambda([PayloadActivityInfo] { return FText::FromString(FString::Printf(TEXT("%u"), PayloadActivityInfo.Cache.PayloadCount)); })
 				];
 
 			Panel->AddSlot(8, Row)
 				[
 					SNew(STextBlock)
 					.Margin(FMargin(ColumnMargin, RowMargin))
-				.ColorAndOpacity(Color)
-				.Font(Font)
-				.Justification(ETextJustify::Center)
-				.Text_Lambda([PayloadAcitvityInfo, BytesToMegaBytes] { return FText::FromString(SingleDecimalFormat((double)PayloadAcitvityInfo.Cache.TotalBytes * BytesToMegaBytes)); })
+					.ColorAndOpacity(Color)
+					.Font(Font)
+					.Justification(ETextJustify::Center)
+					.Text_Lambda([PayloadActivityInfo] { return FText::FromString(SingleDecimalFormat((double)PayloadActivityInfo.Cache.CyclesSpent * FPlatformTime::GetSecondsPerCycle())); })
 				];
 
 			Panel->AddSlot(9, Row)
 				[
 					SNew(STextBlock)
 					.Margin(FMargin(ColumnMargin, RowMargin))
-				.ColorAndOpacity(Color)
-				.Font(Font)
-				.Justification(ETextJustify::Center)
-				.Text_Lambda([PayloadAcitvityInfo] { return FText::FromString(SingleDecimalFormat((double)PayloadAcitvityInfo.Cache.CyclesSpent * FPlatformTime::GetSecondsPerCycle())); })
+					.ColorAndOpacity(Color)
+					.Font(Font)
+					.Justification(ETextJustify::Center)
+					.Text_Lambda([PayloadActivityInfo, BytesToMegaBytes] { return FText::FromString(SingleDecimalFormat((double)PayloadActivityInfo.Cache.TotalBytes * BytesToMegaBytes)); })
 				];
 
 			Row++;
