DDC: Updated the stat dialog layout and split Zen to include its local stats and upstream stats

#rnx

diff --git a/Engine/Source/Developer/DerivedDataCache/Private/ZenDerivedDataBackend.cpp b/Engine/Source/Developer/DerivedDataCache/Private/ZenDerivedDataBackend.cpp
--- a/Engine/Source/Developer/DerivedDataCache/Private/ZenDerivedDataBackend.cpp
+++ b/Engine/Source/Developer/DerivedDataCache/Private/ZenDerivedDataBackend.cpp
@@ -21,6 +21,7 @@
 #include "ZenBackendUtils.h"
 #include "ZenSerialization.h"
 #include "ZenServerHttp.h"
+#include "ZenStatistics.h"
 
 TRACE_DECLARE_INT_COUNTER(ZenDDC_Exist,			TEXT("ZenDDC Exist"));
 TRACE_DECLARE_INT_COUNTER(ZenDDC_ExistHit,		TEXT("ZenDDC Exist Hit"));
@@ -153,6 +154,7 @@ bool FZenDerivedDataBackend::CachedDataProbablyExists(const TCHAR* CacheKey)
 			if (bIsHit)
 			{
 				TRACE_COUNTER_ADD(ZenDDC_ExistHit, int64(1));
+				COOK_STAT(Timer.AddHit(0));
 			}
 			return bIsHit;
 		}
@@ -480,10 +482,53 @@ FZenDerivedDataBackend::GetSpeedClass() const
 
 TSharedRef<FDerivedDataCacheStatsNode> FZenDerivedDataBackend::GatherUsageStats() const
 {
-	TSharedRef<FDerivedDataCacheStatsNode> Usage = MakeShared<FDerivedDataCacheStatsNode>(this, FString::Printf(TEXT("%s.%s"), TEXT("ZenDDC"), *GetName()));
-	Usage->Stats.Add(TEXT(""), UsageStats);
+	Zen::FZenStats ZenStats;
 
-	return Usage;
+	FDerivedDataCacheUsageStats LocalStats;
+	FDerivedDataCacheUsageStats RemoteStats;
+
+#if ENABLE_COOK_STATS
+	using EHitOrMiss = FCookStats::CallStats::EHitOrMiss;
+	using ECacheStatType = FCookStats::CallStats::EStatType;
+
+	Zen::GetDefaultServiceInstance().GetStats(ZenStats);
+
+	const int64 RemotePutSize = int64(ZenStats.UpstreamStats.TotalUploadedMB * 1024 * 1024);
+	const int64 RemoteGetSize = int64(ZenStats.UpstreamStats.TotalDownloadedMB * 1024 * 1024);
+	const int64 LocalGetSize = FMath::Max<int64>(0, UsageStats.GetStats.GetAccumulatedValueAnyThread(EHitOrMiss::Hit, ECacheStatType::Bytes) - RemoteGetSize);
+
+	LocalStats.PutStats = UsageStats.PutStats;
+	LocalStats.ExistsStats = UsageStats.ExistsStats;
+	LocalStats.PrefetchStats = UsageStats.PrefetchStats;
+
+	LocalStats.GetStats.Accumulate(EHitOrMiss::Hit, ECacheStatType::Counter, ZenStats.CacheStats.Hits - ZenStats.CacheStats.UpstreamHits, /*bIsInGameThread*/ false);
+	LocalStats.GetStats.Accumulate(EHitOrMiss::Miss, ECacheStatType::Counter, ZenStats.CacheStats.Misses + ZenStats.CacheStats.UpstreamHits, /*bIsInGameThread*/ false);
+	RemoteStats.GetStats.Accumulate(EHitOrMiss::Hit, ECacheStatType::Counter, ZenStats.CacheStats.UpstreamHits, /*bIsInGameThread*/ false);
+	RemoteStats.GetStats.Accumulate(EHitOrMiss::Miss, ECacheStatType::Counter, ZenStats.CacheStats.Misses, /*bIsInGameThread*/ false);
+
+	LocalStats.GetStats.Accumulate(EHitOrMiss::Hit, ECacheStatType::Bytes, LocalGetSize, /*bIsInGameThread*/ false);
+	RemoteStats.GetStats.Accumulate(EHitOrMiss::Hit, ECacheStatType::Bytes, RemoteGetSize, /*bIsInGameThread*/ false);
+	RemoteStats.PutStats.Accumulate(EHitOrMiss::Hit, ECacheStatType::Bytes, RemotePutSize, /*bIsInGameThread*/ false);
+#endif
+
+	TSharedRef<FDerivedDataCacheStatsNode> LocalNode =
+		MakeShared<FDerivedDataCacheStatsNode>(this, FString::Printf(TEXT("%s.%s"), TEXT("ZenDDC"), *GetName()));
+	LocalNode->Stats.Add(TEXT(""), LocalStats);
+
+	if (ZenStats.UpstreamStats.EndPointStats.IsEmpty())
+	{
+		return LocalNode;
+	}
+
+	TSharedRef<FDerivedDataCacheStatsNode> RemoteNode =
+		MakeShared<FDerivedDataCacheStatsNode>(this, ZenStats.UpstreamStats.EndPointStats[0].Url);
+	RemoteNode->Stats.Add(TEXT(""), RemoteStats);
+
+	TSharedRef<FDerivedDataCacheStatsNode> GroupNode =
+		MakeShared<FDerivedDataCacheStatsNode>(this, TEXT("Zen Group"));
+	GroupNode->Children.Add(LocalNode);
+	GroupNode->Children.Add(RemoteNode);
+	return GroupNode;
 }
 
 bool FZenDerivedDataBackend::TryToPrefetch(TConstArrayView<FString> CacheKeys)
diff --git a/Engine/Source/Editor/DerivedDataEditor/Private/DerivedDataEditorModule.cpp b/Engine/Source/Editor/DerivedDataEditor/Private/DerivedDataEditorModule.cpp
--- a/Engine/Source/Editor/DerivedDataEditor/Private/DerivedDataEditorModule.cpp
+++ b/Engine/Source/Editor/DerivedDataEditor/Private/DerivedDataEditorModule.cpp
@@ -126,16 +126,7 @@ void FDerivedDataEditorModule::ShowResourceUsageTab()
 
 TSharedPtr<SWidget> FDerivedDataEditorModule::CreateCacheStatisticsDialog()
 {
-#if UE_WITH_ZEN
-	if (FDerivedDataInformation::GetHasZenCache())
-	{
-		return SNew(SZenCacheStatisticsDialog);
-	}
-	else
-#endif
-	{
-		return SNew(SDerivedDataCacheStatisticsDialog);
-	}
+	return SNew(SDerivedDataCacheStatisticsDialog);
 }
 
 
diff --git a/Engine/Source/Editor/DerivedDataEditor/Private/SDerivedDataDialogs.cpp b/Engine/Source/Editor/DerivedDataEditor/Private/SDerivedDataDialogs.cpp
--- a/Engine/Source/Editor/DerivedDataEditor/Private/SDerivedDataDialogs.cpp
+++ b/Engine/Source/Editor/DerivedDataEditor/Private/SDerivedDataDialogs.cpp
@@ -1,26 +1,35 @@
 // Copyright Epic Games, Inc. All Rights Reserved.
 
 #include "SDerivedDataDialogs.h"
-#include "DerivedDataInformation.h"
-#include "DerivedDataCacheInterface.h"
+#include "Algo/Sort.h"
 #include "DerivedDataBackendInterface.h"
+#include "DerivedDataCacheInterface.h"
 #include "DerivedDataCacheUsageStats.h"
+#include "DerivedDataInformation.h"
+#include "Framework/Application/SlateApplication.h"
+#include "Internationalization/FastDecimalFormat.h"
+#include "Math/BasicMathExpressionEvaluator.h"
+#include "Math/UnitConversion.h"
+#include "Misc/ExpressionParser.h"
+#include "Styling/StyleColors.h"
+#include "Widgets/Images/SImage.h"
+#include "Widgets/Layout/SGridPanel.h"
 #include "Widgets/SBoxPanel.h"
 #include "Widgets/SToolTip.h"
 #include "Widgets/Text/STextBlock.h"
-#include "Widgets/SBoxPanel.h"
-#include "Widgets/Layout/SGridPanel.h"
-#include "Widgets/Images/SImage.h"
-#include "Styling/StyleColors.h"
-#include "Math/UnitConversion.h"
-#include "Math/BasicMathExpressionEvaluator.h"
-#include "Misc/ExpressionParser.h"
-#include "Framework/Application/SlateApplication.h"
-#include "Internationalization/FastDecimalFormat.h"
-
 
 #define LOCTEXT_NAMESPACE "DerivedDataCacheEditor"
 
+template <typename ValueType>
+static FString ZeroDecimalFormat(ValueType Value)
+{
+	const FNumberFormattingOptions NumberFormattingOptions = FNumberFormattingOptions()
+		.SetUseGrouping(true)
+		.SetMinimumFractionalDigits(0)
+		.SetMaximumFractionalDigits(0);
+	return FastDecimalFormat::NumberToString(Value, ExpressionParser::GetLocalizedNumberFormattingRules(), NumberFormattingOptions);
+}
+
 FString SingleDecimalFormat(double Value)
 {
 	const FNumberFormattingOptions NumberFormattingOptions = FNumberFormattingOptions()
@@ -28,7 +37,7 @@ FString SingleDecimalFormat(double Value)
 		.SetMinimumFractionalDigits(1)
 		.SetMaximumFractionalDigits(1);
 	return FastDecimalFormat::NumberToString(Value, ExpressionParser::GetLocalizedNumberFormattingRules(), NumberFormattingOptions);
-};
+}
 
 void SDerivedDataRemoteStoreDialog::Construct(const FArguments& InArgs)
 {
@@ -79,116 +88,64 @@ TSharedRef<SWidget> SDerivedDataRemoteStoreDialog::GetGridPanel()
 
 	const double DownloadedBytesMB = FUnitConversion::Convert(FDerivedDataInformation::GetCacheActivitySizeBytes(true, false), EUnit::Bytes, EUnit::Megabytes);
 	const double UploadedBytesMB = FUnitConversion::Convert(FDerivedDataInformation::GetCacheActivitySizeBytes(false, false), EUnit::Bytes, EUnit::Megabytes);
-	
+
 	TSharedRef<SGridPanel> Panel =
 		SNew(SGridPanel);
-	
-	int32 Row = 0;
-
-	/*Panel->AddSlot(0, Row)
-		[
-			SNew(STextBlock)
-			.ColorAndOpacity(FStyleColors::Foreground)
-			.Text(LOCTEXT("Efficiency", "Efficiency"))
-		];
-
-	Panel->AddSlot(1, Row)
-		[
-			SNew(STextBlock)
-			.Margin(FMargin(5, 0))
-			.Text(FText::FromString(SingleDecimalFormat(Efficiency * 100.0f) + TEXT(" %")))
-			.Justification(ETextJustify::Right)
-		];
-
-	Row++;
-
-	Panel->AddSlot(0, Row)
-		[
-			SNew(STextBlock)
-			.ColorAndOpacity(FStyleColors::Foreground)
-			.Text(LOCTEXT("Loaded", "Loaded"))
-		];
-
-	Panel->AddSlot(1, Row)
-		[
-			SNew(STextBlock)
-			.Margin(FMargin(5, 0))
-			.Text(FText::FromString(SingleDecimalFormat(DDCResourceStatsTotal.LoadSizeMB) + TEXT(" MB")))
-			.Justification(ETextJustify::Right)
-		];
 
-	Row++;
-	
-	Panel->AddSlot(0, Row)
-		[
-			SNew(STextBlock)
-			.ColorAndOpacity(FStyleColors::Foreground)
-			.Text(LOCTEXT("Built", "Built"))
-		];
-
-	Panel->AddSlot(1, Row)
-		[
-			SNew(STextBlock)
-			.Margin(FMargin(5, 0))
-			.Text(FText::FromString(SingleDecimalFormat(DDCResourceStatsTotal.BuildSizeMB) + TEXT(" MB")))
-			.Justification(ETextJustify::Right)
-		];
-
-	Row++;*/
+	int32 Row = 0;
 
 	Panel->AddSlot(0, Row)
-		[
+	[
 		SNew(STextBlock)
 		.Font(FCoreStyle::GetDefaultFontStyle("Bold", 10))
 		.ColorAndOpacity(FStyleColors::Foreground)
 		.Text(LOCTEXT("Remote Storage", "Remote Storage"))
-
-		];
+	];
 
 	Row++;
 
 	Panel->AddSlot(0, Row)
-		[
-			SNew(STextBlock)
-			.Text(LOCTEXT("Connected", "Connected"))
-		];
+	[
+		SNew(STextBlock)
+		.Text(LOCTEXT("Connected", "Connected"))
+	];
 
 	Panel->AddSlot(1, Row)
-		[
-			SNew(STextBlock)
-			.Text(FText::FromString(FDerivedDataInformation::GetHasRemoteCache()? TEXT("True") : TEXT("False")))
-			.Justification(ETextJustify::Right)
-		];
+	.HAlign(HAlign_Right)
+	[
+		SNew(STextBlock)
+		.Text(FDerivedDataInformation::GetHasRemoteCache() ? LOCTEXT("True", "True") : LOCTEXT("False", "False"))
+	];
 
 	Row++;
 
 	Panel->AddSlot(0, Row)
-		[
-			SNew(STextBlock)
-			.Text(LOCTEXT("Downloaded", "Downloaded"))
-		];
+	[
+		SNew(STextBlock)
+		.Text(LOCTEXT("Downloaded", "Downloaded"))
+	];
 
 	Panel->AddSlot(1, Row)
-		[
-			SNew(STextBlock)
-			.Text(FText::FromString(SingleDecimalFormat(DownloadedBytesMB) + TEXT(" MB")))
-			.Justification(ETextJustify::Right)
-		];
+	.HAlign(HAlign_Right)
+	[
+		SNew(STextBlock)
+		.Text(FText::FromString(SingleDecimalFormat(DownloadedBytesMB) + TEXT(" MiB")))
+	];
 
 	Row++;
 
 	Panel->AddSlot(0, Row)
-		[
-			SNew(STextBlock)
-			.Text(LOCTEXT("Uploaded", "Uploaded"))
-		];
+	[
+		SNew(STextBlock)
+		.Text(LOCTEXT("Uploaded", "Uploaded"))
+	];
 
 	Panel->AddSlot(1, Row)
-		[
-			SNew(STextBlock)
-			.Text(FText::FromString(SingleDecimalFormat(UploadedBytesMB) + TEXT(" MB")))
-			.Justification(ETextJustify::Right)
-		];
+	.HAlign(HAlign_Right)
+	[
+		SNew(STextBlock)
+		.Text(FText::FromString(SingleDecimalFormat(UploadedBytesMB) + TEXT(" MiB")))
+	];
 
 	Row++;
 
@@ -202,14 +159,13 @@ void SDerivedDataResourceUsageDialog::Construct(const FArguments& InArgs)
 	this->ChildSlot
 	[
 		SNew(SVerticalBox)
-
 		+ SVerticalBox::Slot()
-			.AutoHeight()
-			.Padding(0, 20, 0, 0)
-			.Expose(GridSlot)
-			[
-				GetGridPanel()
-			]
+		.AutoHeight()
+		.Padding(0, 20, 0, 0)
+		.Expose(GridSlot)
+		[
+			GetGridPanel()
+		]
 	];
 
 	RegisterActiveTimer(0.5f, FWidgetActiveTimerDelegate::CreateSP(this, &SDerivedDataResourceUsageDialog::UpdateGridPanels));
@@ -228,15 +184,26 @@ EActiveTimerReturnType SDerivedDataResourceUsageDialog::UpdateGridPanels(double
 }
 
 TSharedRef<SWidget> SDerivedDataResourceUsageDialog::GetGridPanel()
-{	
-	
+{
 	TArray<FDerivedDataCacheResourceStat> DDCResourceStats;
 
 	// Grab the resource stats
 	GetDerivedDataCacheRef().GatherResourceStats(DDCResourceStats);
 
-	// Sort results on asscending Load size
-	DDCResourceStats.Sort([](const FDerivedDataCacheResourceStat& LHS, const FDerivedDataCacheResourceStat& RHS) { return LHS.BuildSizeMB > RHS.BuildSizeMB; });
+	// Sort results on descending build size, then descending load size, then ascending asset type.
+	const auto CompareStats = [](const FDerivedDataCacheResourceStat& A, const FDerivedDataCacheResourceStat& B)
+	{
+		if (A.BuildSizeMB != B.BuildSizeMB)
+		{
+			return A.BuildSizeMB > B.BuildSizeMB;
+		}
+		if (A.LoadSizeMB != B.LoadSizeMB)
+		{
+			return A.LoadSizeMB > B.LoadSizeMB;
+		}
+		return A.AssetType.Compare(B.AssetType, ESearchCase::IgnoreCase) < 0;
+	};
+	Algo::Sort(DDCResourceStats, CompareStats);
 
 	FDerivedDataCacheResourceStat DDCResourceStatsTotal(TEXT("Total"));
 
@@ -252,251 +219,247 @@ TSharedRef<SWidget> SDerivedDataResourceUsageDialog::GetGridPanel()
 	int32 Row = 0;
 
 	const float RowMargin = 0.0f;
-	const float TitleMargin = 10.0f;
 	const float ColumnMargin = 10.0f;
+	const FMargin TitleMargin(0.0f, 10.0f, ColumnMargin, 10.0f);
+	const FMargin TitleMarginFirstColumn(ColumnMargin, 10.0f);
 	const FSlateColor TitleColor = FStyleColors::AccentWhite;
 	const FSlateFontInfo TitleFont = FCoreStyle::GetDefaultFontStyle("Bold", 10);
-
+	const FMargin DefaultMargin(0.0f, RowMargin, ColumnMargin, RowMargin);
+	const FMargin DefaultMarginFirstColumn(ColumnMargin, RowMargin);
 
 	Panel->AddSlot(2, Row)
-		[
+	.HAlign(HAlign_Center)
+	[
 		SNew(STextBlock)
-		.Margin(FMargin(ColumnMargin, RowMargin))
+		.Margin(DefaultMargin)
 		.ColorAndOpacity(TitleColor)
 		.Font(TitleFont)
-		.Justification(ETextJustify::Center)
 		.Text(LOCTEXT("Loaded", "Loaded"))
-		];
-
+	];
 
 	Panel->AddSlot(5, Row)
-		[
+	.HAlign(HAlign_Center)
+	[
 		SNew(STextBlock)
-		.Margin(FMargin(ColumnMargin, RowMargin))
+		.Margin(DefaultMargin)
 		.ColorAndOpacity(TitleColor)
 		.Font(TitleFont)
-		.Justification(ETextJustify::Center)
 		.Text(LOCTEXT("Built", "Built"))
-		];
+	];
 
 	Row++;
 
 	Panel->AddSlot(0, Row)
-		[
+	[
 		SNew(STextBlock)
-		.Margin(FMargin(ColumnMargin, RowMargin, 0.0f, TitleMargin))
+		.Margin(TitleMarginFirstColumn)
 		.ColorAndOpacity(TitleColor)
 		.Font(TitleFont)
-		.Justification(ETextJustify::Left)
 		.Text(LOCTEXT("ResourceType", "Resource Type"))
-		];
+	];
 
 	Panel->AddSlot(1, Row)
-		[
-	
+	.HAlign(HAlign_Right)
+	[
 		SNew(STextBlock)
-		.Margin(FMargin(ColumnMargin, RowMargin, 0.0f, TitleMargin))
+		.Margin(TitleMargin)
 		.ColorAndOpacity(TitleColor)
 		.Font(TitleFont)
-		.Justification(ETextJustify::Center)
 		.Text(LOCTEXT("Count", "Count"))
-		];
+	];
 
 	Panel->AddSlot(2, Row)
-		[
+	.HAlign(HAlign_Right)
+	[
 		SNew(STextBlock)
-		.Margin(FMargin(ColumnMargin, RowMargin, 0.0f, TitleMargin))
+		.Margin(TitleMargin)
 		.ColorAndOpacity(TitleColor)
 		.Font(TitleFont)
-		.Justification(ETextJustify::Center)
 		.Text(LOCTEXT("Time (Sec)", "Time (Sec)"))
-		];
+	];
 
 	Panel->AddSlot(3, Row)
-		[
+	.HAlign(HAlign_Right)
+	[
 		SNew(STextBlock)
-		.Margin(FMargin(ColumnMargin, RowMargin, 0.0f, TitleMargin))
+		.Margin(TitleMargin)
 		.ColorAndOpacity(TitleColor)
 		.Font(TitleFont)
-		.Justification(ETextJustify::Center)
-		.Text(LOCTEXT("Size (MB)", "Size (MB)"))
-		];
+		.Text(LOCTEXT("Size (MiB)", "Size (MiB)"))
+	];
 
 	Panel->AddSlot(4, Row)
-		[
+	.HAlign(HAlign_Right)
+	[
 		SNew(STextBlock)
-		.Margin(FMargin(ColumnMargin, RowMargin, 0.0f, TitleMargin))
+		.Margin(TitleMargin)
 		.ColorAndOpacity(TitleColor)
 		.Font(TitleFont)
-		.Justification(ETextJustify::Center)
 		.Text(LOCTEXT("Count", "Count"))
-		];
+	];
 
 	Panel->AddSlot(5, Row)
-		[
+	.HAlign(HAlign_Right)
+	[
 		SNew(STextBlock)
-		.Margin(FMargin(ColumnMargin, RowMargin, 0.0f, TitleMargin))
+		.Margin(TitleMargin)
 		.ColorAndOpacity(TitleColor)
 		.Font(TitleFont)
-		.Justification(ETextJustify::Center)
 		.Text(LOCTEXT("Time (Sec)", "Time (Sec)"))
-		];
+	];
 
 	Panel->AddSlot(6, Row)
-		[
+	.HAlign(HAlign_Right)
+	[
 		SNew(STextBlock)
-		.Margin(FMargin(ColumnMargin, RowMargin, 0.0f, TitleMargin))
+		.Margin(TitleMargin)
 		.ColorAndOpacity(TitleColor)
 		.Font(TitleFont)
-		.Justification(ETextJustify::Center)
-		.Text(LOCTEXT("Size (MB)", "Size (MB)"))
-		];
+		.Text(LOCTEXT("Size (MiB)", "Size (MiB)"))
+	];
 
 	Row++;
 
-
 	for (const FDerivedDataCacheResourceStat& Stat : DDCResourceStats)
 	{	
 		Panel->AddSlot(0, Row)
 		[
 			SNew(STextBlock)
-			.Margin(FMargin(ColumnMargin, RowMargin))
+			.Margin(DefaultMarginFirstColumn)
 			.Text(FText::FromString(Stat.AssetType))
-			.Justification(ETextJustify::Left)
 		];
 
 		Panel->AddSlot(1, Row)
+		.HAlign(HAlign_Right)
 		[
 			SNew(STextBlock)
-			.Margin(FMargin(ColumnMargin, RowMargin))
-			.Text(FText::FromString(FString::Printf(TEXT("%u"),Stat.LoadCount)))
-			.Justification(ETextJustify::Center)
+			.Margin(DefaultMargin)
+			.Text(FText::FromString(ZeroDecimalFormat(Stat.LoadCount)))
 		];
 
 		Panel->AddSlot(2, Row)
+		.HAlign(HAlign_Right)
 		[
 			SNew(STextBlock)
-			.Margin(FMargin(ColumnMargin, RowMargin))
+			.Margin(DefaultMargin)
 			.Text(FText::FromString(SingleDecimalFormat(Stat.LoadTimeSec)))
-			.Justification(ETextJustify::Center)
 		];
 
 		Panel->AddSlot(3, Row)
+		.HAlign(HAlign_Right)
 		[
 			SNew(STextBlock)
-			.Margin(FMargin(ColumnMargin, RowMargin))
+			.Margin(DefaultMargin)
 			.Text(FText::FromString(SingleDecimalFormat(Stat.LoadSizeMB)))
-			.Justification(ETextJustify::Center)
 		];
 
 		Panel->AddSlot(4, Row)
+		.HAlign(HAlign_Right)
 		[
 			SNew(STextBlock)
-			.Margin(FMargin(ColumnMargin, RowMargin))
-			.Text(FText::FromString(FString::Printf(TEXT("%u"),Stat.BuildCount)))
-			.Justification(ETextJustify::Center)
+			.Margin(DefaultMargin)
+			.Text(FText::FromString(ZeroDecimalFormat(Stat.BuildCount)))
 		];
 
 		Panel->AddSlot(5, Row)
+		.HAlign(HAlign_Right)
 		[
 			SNew(STextBlock)
-			.Margin(FMargin(ColumnMargin, RowMargin))
+			.Margin(DefaultMargin)
 			.Text(FText::FromString(SingleDecimalFormat(Stat.BuildTimeSec)))
-			.Justification(ETextJustify::Center)
 		];
 
 		Panel->AddSlot(6, Row)
+		.HAlign(HAlign_Right)
 		[
 			SNew(STextBlock)
-			.Margin(FMargin(ColumnMargin, RowMargin))
+			.Margin(DefaultMargin)
 			.Text(FText::FromString(SingleDecimalFormat(Stat.BuildSizeMB)))
-			.Justification(ETextJustify::Center)
 		];
 
 		Row++;
 	}
 	
 	Panel->AddSlot(0, Row)
-		[
-			SNew(STextBlock)
-			.Margin(FMargin(ColumnMargin, RowMargin))
-			.ColorAndOpacity(TitleColor)
-			.Font(TitleFont)
-			.Text(FText::FromString(DDCResourceStatsTotal.AssetType))
-			.Font(FCoreStyle::GetDefaultFontStyle("Bold", 10))
-			.Justification(ETextJustify::Left)
-		];
+	[
+		SNew(STextBlock)
+		.Margin(TitleMarginFirstColumn)
+		.ColorAndOpacity(TitleColor)
+		.Font(TitleFont)
+		.Text(FText::FromString(DDCResourceStatsTotal.AssetType))
+	];
 
 	Panel->AddSlot(1, Row)
-		[
-			SNew(STextBlock)
-			.Margin(FMargin(ColumnMargin, RowMargin))
-			.ColorAndOpacity(TitleColor)
-			.Font(TitleFont)
-			.Justification(ETextJustify::Center)
-			.Text(FText::FromString(FString::Printf(TEXT("%u"), DDCResourceStatsTotal.LoadCount)))
-		];
+	.HAlign(HAlign_Right)
+	[
+		SNew(STextBlock)
+		.Margin(TitleMargin)
+		.ColorAndOpacity(TitleColor)
+		.Font(TitleFont)
+		.Text(FText::FromString(ZeroDecimalFormat(DDCResourceStatsTotal.LoadCount)))
+	];
 
 	Panel->AddSlot(2, Row)
-		[
-			SNew(STextBlock)
-			.Margin(FMargin(ColumnMargin, RowMargin))
-			.ColorAndOpacity(TitleColor)
-			.Font(TitleFont)
-			.Justification(ETextJustify::Center)
-			.Text(FText::FromString(SingleDecimalFormat(DDCResourceStatsTotal.LoadTimeSec)))
-		];
+	.HAlign(HAlign_Right)
+	[
+		SNew(STextBlock)
+		.Margin(TitleMargin)
+		.ColorAndOpacity(TitleColor)
+		.Font(TitleFont)
+		.Text(FText::FromString(SingleDecimalFormat(DDCResourceStatsTotal.LoadTimeSec)))
+	];
 
 	Panel->AddSlot(3, Row)
-		[
-			SNew(STextBlock)
-			.Margin(FMargin(ColumnMargin, RowMargin))
-			.ColorAndOpacity(TitleColor)
-			.Font(TitleFont)
-			.Justification(ETextJustify::Center)
-			.Text(FText::FromString(SingleDecimalFormat(DDCResourceStatsTotal.LoadSizeMB)))
-		];
+	.HAlign(HAlign_Right)
+	[
+		SNew(STextBlock)
+		.Margin(TitleMargin)
+		.ColorAndOpacity(TitleColor)
+		.Font(TitleFont)
+		.Text(FText::FromString(SingleDecimalFormat(DDCResourceStatsTotal.LoadSizeMB)))
+	];
 
 	Panel->AddSlot(4, Row)
-		[
-			SNew(STextBlock)
-			.Margin(FMargin(ColumnMargin, RowMargin))
-			.ColorAndOpacity(TitleColor)
-			.Font(TitleFont)
-			.Justification(ETextJustify::Center)
-			.Text(FText::FromString(FString::Printf(TEXT("%u"), DDCResourceStatsTotal.BuildCount)))
-		];
+	.HAlign(HAlign_Right)
+	[
+		SNew(STextBlock)
+		.Margin(TitleMargin)
+		.ColorAndOpacity(TitleColor)
+		.Font(TitleFont)
+		.Text(FText::FromString(ZeroDecimalFormat(DDCResourceStatsTotal.BuildCount)))
+	];
 
 	Panel->AddSlot(5, Row)
-		[
-			SNew(STextBlock)
-			.Margin(FMargin(ColumnMargin, RowMargin))
-			.ColorAndOpacity(TitleColor)
-			.Font(TitleFont)
-			.Justification(ETextJustify::Center)
-			.Text(FText::FromString(SingleDecimalFormat(DDCResourceStatsTotal.BuildTimeSec)))
-		];
+	.HAlign(HAlign_Right)
+	[
+		SNew(STextBlock)
+		.Margin(TitleMargin)
+		.ColorAndOpacity(TitleColor)
+		.Font(TitleFont)
+		.Text(FText::FromString(SingleDecimalFormat(DDCResourceStatsTotal.BuildTimeSec)))
+	];
 
 	Panel->AddSlot(6, Row)
-		[
-			SNew(STextBlock)
-			.Margin(FMargin(ColumnMargin, RowMargin))
-			.ColorAndOpacity(TitleColor)
-			.Font(TitleFont)
-			.Justification(ETextJustify::Center)
-			.Text(FText::FromString(SingleDecimalFormat(DDCResourceStatsTotal.BuildSizeMB)))
-		];
+	.HAlign(HAlign_Right)
+	[
+		SNew(STextBlock)
+		.Margin(TitleMargin)
+		.ColorAndOpacity(TitleColor)
+		.Font(TitleFont)
+		.Text(FText::FromString(SingleDecimalFormat(DDCResourceStatsTotal.BuildSizeMB)))
+	];
 
 	return Panel;
 }
 
 
+
 void SDerivedDataCacheStatisticsDialog::Construct(const FArguments& InArgs)
 {
 	const float RowMargin = 0.0f;
 	const float TitleMargin = 10.0f;
 	const float ColumnMargin = 10.0f;
-	const FSlateColor TitleColour = FStyleColors::AccentWhite;
+	const FSlateColor TitleColor = FStyleColors::AccentWhite;
 	const FSlateFontInfo TitleFont = FCoreStyle::GetDefaultFontStyle("Bold", 10);
 
 	this->ChildSlot
@@ -506,17 +469,17 @@ void SDerivedDataCacheStatisticsDialog::Construct(const FArguments& InArgs)
 		.Padding(0, 20, 0, 0)
 		.AutoHeight()
 		[
-			SNew(SHorizontalBox)		
+			SNew(SHorizontalBox)
 			+SHorizontalBox::Slot()
 			.FillWidth(1.0f)
 			[
 				SNew(STextBlock)
-				.Margin(FMargin(ColumnMargin, RowMargin, 0.0f, TitleMargin))
-				.ColorAndOpacity(TitleColour)
+				.Margin(TitleMargin)
+				.ColorAndOpacity(TitleColor)
 				.Font(TitleFont)
 				.Justification(ETextJustify::Left)
-				.Text( FText::FromString(GetDerivedDataCache()->GetGraphName()))
-			]		
+				.Text(FText::FromString(GetDerivedDataCache()->GetGraphName()))
+			]
 		]
 		+ SVerticalBox::Slot()
 		.AutoHeight()
@@ -525,7 +488,6 @@ void SDerivedDataCacheStatisticsDialog::Construct(const FArguments& InArgs)
 		[
 			GetGridPanel()
 		]
-
 	];
 
 	RegisterActiveTimer(0.5f, FWidgetActiveTimerDelegate::CreateSP(this, &SDerivedDataCacheStatisticsDialog::UpdateGridPanels));
@@ -541,18 +503,19 @@ EActiveTimerReturnType SDerivedDataCacheStatisticsDialog::UpdateGridPanels(doubl
 	SlatePrepass(GetPrepassLayoutScaleMultiplier());
 
 	return EActiveTimerReturnType::Continue;
-} 
+}
 
 
 
 TSharedRef<SWidget> SDerivedDataCacheStatisticsDialog::GetGridPanel()
 {
-	PRAGMA_DISABLE_DEPRECATION_WARNINGS
+	PRAGMA_DISABLE_DEPRECATION_WARNINGS;
 	TSharedRef<FDerivedDataCacheStatsNode> RootUsage = GetDerivedDataCache()->GatherUsageStats();
-	PRAGMA_ENABLE_DEPRECATION_WARNINGS
+	PRAGMA_ENABLE_DEPRECATION_WARNINGS;
 	TArray<TSharedRef<const FDerivedDataCacheStatsNode>> LeafUsageStats;
-	RootUsage->ForEachDescendant([&LeafUsageStats](TSharedRef<const FDerivedDataCacheStatsNode> Node) {
-		if (Node->Children.Num() == 0)
+	RootUsage->ForEachDescendant([&LeafUsageStats](TSharedRef<const FDerivedDataCacheStatsNode> Node)
+	{
+		if (Node->Children.IsEmpty())
 		{
 			LeafUsageStats.Add(Node);
 		}
@@ -564,15 +527,18 @@ TSharedRef<SWidget> SDerivedDataCacheStatisticsDialog::GetGridPanel()
 	int32 Row = 0;
 
 	const float RowMargin = 0.0f;
-	const float TitleMargin = 10.0f;
 	const float ColumnMargin = 10.0f;
+	const FMargin TitleMargin(0.0f, 10.0f, ColumnMargin, 10.0f);
+	const FMargin TitleMarginFirstColumn(ColumnMargin, 10.0f);
 	const FSlateColor TitleColor = FStyleColors::AccentWhite;
 	const FSlateFontInfo TitleFont = FCoreStyle::GetDefaultFontStyle("Bold", 10);
+	const FMargin DefaultMargin(0.0f, RowMargin, ColumnMargin, RowMargin);
+	const FMargin DefaultMarginFirstColumn(ColumnMargin, RowMargin);
 
 	Panel->AddSlot(0, Row)
 	[
 		SNew(STextBlock)
-		.Margin(FMargin(ColumnMargin, RowMargin, 0.0f, TitleMargin))
+		.Margin(TitleMarginFirstColumn)
 		.Font(FCoreStyle::GetDefaultFontStyle("Bold", 10))
 		.Text(LOCTEXT("Cache", "Cache"))
 		.ColorAndOpacity(TitleColor)
@@ -582,66 +548,54 @@ TSharedRef<SWidget> SDerivedDataCacheStatisticsDialog::GetGridPanel()
 	Panel->AddSlot(1, Row)
 	[
 		SNew(STextBlock)
-		.Margin(FMargin(ColumnMargin, RowMargin, 0.0f, TitleMargin))
+		.Margin(TitleMargin)
 		.ColorAndOpacity(TitleColor)
 		.Font(TitleFont)
-		.Justification(ETextJustify::Left)
 		.Text(LOCTEXT("Location", "Location"))
 	];
 
 	Panel->AddSlot(2, Row)
+	.HAlign(HAlign_Right)
 	[
 		SNew(STextBlock)
-		.Margin(FMargin(ColumnMargin, RowMargin, 0.0f, TitleMargin))
+		.Margin(TitleMargin)
 		.ColorAndOpacity(TitleColor)
 		.Font(TitleFont)
-		.Justification(ETextJustify::Left)
 		.Text(LOCTEXT("HitPercentage", "Hit%"))
 	];
 
 	Panel->AddSlot(3, Row)
+	.HAlign(HAlign_Right)
 	[
 		SNew(STextBlock)
-		.Margin(FMargin(ColumnMargin, RowMargin, 0.0f, TitleMargin))
+		.Margin(TitleMargin)
 		.ColorAndOpacity(TitleColor)
 		.Font(TitleFont)
-		.Justification(ETextJustify::Left)
-		.Text(LOCTEXT("Read", "Read"))
+		.Text(LOCTEXT("Read", "Read (MiB)"))
 	];
 
 	Panel->AddSlot(4, Row)
+	.HAlign(HAlign_Right)
 	[
 		SNew(STextBlock)
-		.Margin(FMargin(ColumnMargin, RowMargin, 0.0f, TitleMargin))
+		.Margin(TitleMargin)
 		.ColorAndOpacity(TitleColor)
 		.Font(TitleFont)
-		.Justification(ETextJustify::Center)
-		.Text(LOCTEXT("Read", "Read"))
+		.Text(LOCTEXT("Write", "Write (MiB)"))
 	];
 
 	Panel->AddSlot(5, Row)
 	[
 		SNew(STextBlock)
-		.Margin(FMargin(ColumnMargin, RowMargin, 0.0f, TitleMargin))
+		.Margin(TitleMargin)
 		.ColorAndOpacity(TitleColor)
 		.Font(TitleFont)
-		.Justification(ETextJustify::Left)
-		.Text(LOCTEXT("Write", "Write"))
-	];
-
-	Panel->AddSlot(6, Row)
-	[
-		SNew(STextBlock)
-		.Margin(FMargin(ColumnMargin, RowMargin, 0.0f, TitleMargin))
-		.ColorAndOpacity(TitleColor)
-		.Font(TitleFont)
-		.Justification(ETextJustify::Left)
 		.Text(LOCTEXT("Details", "Details"))
 	];
 
 	Row++;
 
-	double SumTotalGetMB =0.0;
+	double SumTotalGetMB = 0.0;
 	double SumTotalPutMB = 0.0;
 
 	for (TSharedRef<const FDerivedDataCacheStatsNode> Node : LeafUsageStats)
@@ -649,94 +603,73 @@ TSharedRef<SWidget> SDerivedDataCacheStatisticsDialog::GetGridPanel()
 		const FDerivedDataBackendInterface* Backend = Node->GetBackendInterface();
 
 		if (Backend->GetDisplayName().Equals("Memory") || Backend->IsWrapper())
+		{
 			continue;
+		}
 
-		TSharedRef<FDerivedDataCacheStatsNode> Usage = Backend->GatherUsageStats();
 		FDerivedDataCacheUsageStats Stats;
 
-		for (const auto& KVP : Usage->Stats)
+		for (const auto& KVP : Node->Stats)
 		{
 			Stats.Combine(KVP.Value);
 		}
 
-		
 		const int64 TotalGetBytes = Stats.GetStats.GetAccumulatedValueAnyThread(FCookStats::CallStats::EHitOrMiss::Hit, FCookStats::CallStats::EStatType::Bytes);
 		const int64 TotalPutBytes = Stats.PutStats.GetAccumulatedValueAnyThread(FCookStats::CallStats::EHitOrMiss::Hit, FCookStats::CallStats::EStatType::Bytes);
-		const int64 TotalPrefetchBytes = Stats.PrefetchStats.GetAccumulatedValueAnyThread(FCookStats::CallStats::EHitOrMiss::Hit, FCookStats::CallStats::EStatType::Bytes);
 
-		const int64 TotalGets_Hit = Stats.GetStats.GetAccumulatedValueAnyThread(FCookStats::CallStats::EHitOrMiss::Hit, FCookStats::CallStats::EStatType::Counter);
-		const int64 TotalGets_Miss = Stats.GetStats.GetAccumulatedValueAnyThread(FCookStats::CallStats::EHitOrMiss::Miss, FCookStats::CallStats::EStatType::Counter);
+		const int64 TotalGetHits = Stats.GetStats.GetAccumulatedValueAnyThread(FCookStats::CallStats::EHitOrMiss::Hit, FCookStats::CallStats::EStatType::Counter);
+		const int64 TotalGetMisses = Stats.GetStats.GetAccumulatedValueAnyThread(FCookStats::CallStats::EHitOrMiss::Miss, FCookStats::CallStats::EStatType::Counter);
+		const int64 TotalRequests = TotalGetHits + TotalGetMisses;
+		const double HitRate = TotalRequests > 0 ? 100.0 * TotalGetHits / TotalRequests : 0.0;
 
-		const int64 TotalRequests = TotalGets_Hit + TotalGets_Miss;
+		const double TotalGetMB = FUnitConversion::Convert(TotalGetBytes, EUnit::Bytes, EUnit::Megabytes);
+		const double TotalPutMB = FUnitConversion::Convert(TotalPutBytes, EUnit::Bytes, EUnit::Megabytes);
+
+		SumTotalGetMB += TotalGetMB;
+		SumTotalPutMB += TotalPutMB;
 
 		Panel->AddSlot(0, Row)
-			.HAlign(HAlign_Left)
 		[
 			SNew(STextBlock)
-			.Margin(FMargin(ColumnMargin, RowMargin))
-			.Justification(ETextJustify::Left)
+			.Margin(DefaultMarginFirstColumn)
 			.Text(FText::FromString(Backend->GetDisplayName()))
 		];
 
 		Panel->AddSlot(1, Row)
 		[
 			SNew(STextBlock)
-			.Margin(FMargin(ColumnMargin, RowMargin))
-			.Justification(ETextJustify::Left)
-			.Text( FText::FromString((Backend->GetSpeedClass() == FDerivedDataBackendInterface::ESpeedClass::Local) ? "Local" : "Remote" ))
+			.Margin(DefaultMargin)
+			.Text((Backend->GetSpeedClass() == FDerivedDataBackendInterface::ESpeedClass::Local) ? LOCTEXT("Local", "Local") : LOCTEXT("Remote", "Remote"))
 		];
 
-		const double HitPercentage = TotalRequests > 0 ? (100.0 * (TotalGets_Hit / (double)TotalRequests)) : 0.0;
-
 		Panel->AddSlot(2, Row)
-			.HAlign(HAlign_Left)
-			[
+		.HAlign(HAlign_Right)
+		[
 			SNew(STextBlock)
-			.Margin(FMargin(ColumnMargin, RowMargin))
-			.Justification(ETextJustify::Left)
-			.Text(FText::FromString(SingleDecimalFormat(HitPercentage) + TEXT(" %")))
-			];
-
-		const double TotalGetMB = FUnitConversion::Convert(TotalGetBytes, EUnit::Bytes, EUnit::Megabytes);
-
-		SumTotalGetMB += TotalGetMB;
+			.Margin(DefaultMargin)
+			.Text(FText::FromString(SingleDecimalFormat(HitRate)))
+		];
 
 		Panel->AddSlot(3, Row)
-		.HAlign(HAlign_Left)
+		.HAlign(HAlign_Right)
 		[
 			SNew(STextBlock)
-			.Margin(FMargin(ColumnMargin, RowMargin))
-			.Justification(ETextJustify::Left)
-			.Text(FText::FromString(SingleDecimalFormat(TotalGetMB) + TEXT(" MB")))
+			.Margin(DefaultMargin)
+			.Text(FText::FromString(SingleDecimalFormat(TotalGetMB)))
 		];
 
-		const double TotalPutMB = FUnitConversion::Convert(TotalPutBytes, EUnit::Bytes, EUnit::Megabytes);
-
-		SumTotalPutMB += TotalPutMB;
-
 		Panel->AddSlot(4, Row)
-		.HAlign(HAlign_Left)
+		.HAlign(HAlign_Right)
 		[
 			SNew(STextBlock)
-			.Margin(FMargin(ColumnMargin, RowMargin))
-			.Justification(ETextJustify::Left)
-			.Text(FText::FromString(SingleDecimalFormat(TotalPutMB) + TEXT(" MB")))
+			.Margin(DefaultMargin)
+			.Text(FText::FromString(SingleDecimalFormat(TotalPutMB)))
 		];
 
 		Panel->AddSlot(5, Row)
 		[
 			SNew(STextBlock)
-			.Margin(FMargin(ColumnMargin, RowMargin))
-			.Justification(ETextJustify::Right)
-			.Text(FText::FromString(SingleDecimalFormat(TotalPutMB) + TEXT(" MB")))
-		];
-
-		Panel->AddSlot(6, Row)
-		.HAlign(HAlign_Left)
-		[
-			SNew(STextBlock)
-			.Margin(FMargin(ColumnMargin, RowMargin))
-			.Justification(ETextJustify::Left)
+			.Margin(DefaultMargin)
 			.Text(FText::FromString(Backend->GetName()))
 		];
 
@@ -744,38 +677,35 @@ TSharedRef<SWidget> SDerivedDataCacheStatisticsDialog::GetGridPanel()
 	}
 
 	Panel->AddSlot(0, Row)
-		[
-			SNew(STextBlock)
-			.Margin(FMargin(ColumnMargin, RowMargin))
-			.ColorAndOpacity(TitleColor)
-			.Font(TitleFont)
-			.Justification(ETextJustify::Left)
-			.Text(FText::FromString(TEXT("Total")))	
-		];
+	[
+		SNew(STextBlock)
+		.Margin(TitleMarginFirstColumn)
+		.ColorAndOpacity(TitleColor)
+		.Font(TitleFont)
+		.Text(LOCTEXT("Total", "Total"))
+	];
 
 	Panel->AddSlot(3, Row)
-		.HAlign(HAlign_Right)
-		[
+	.HAlign(HAlign_Right)
+	[
 		SNew(STextBlock)
-		.Margin(FMargin(ColumnMargin, RowMargin))
+		.Margin(TitleMargin)
 		.ColorAndOpacity(TitleColor)
 		.Font(TitleFont)
-		.Justification(ETextJustify::Right)
-		.Text(FText::FromString(SingleDecimalFormat(SumTotalGetMB) + TEXT(" MB")))
-		];
+		.Text(FText::FromString(SingleDecimalFormat(SumTotalGetMB)))
+	];
 
 	Panel->AddSlot(4, Row)
-		.HAlign(HAlign_Right)
-		[
+	.HAlign(HAlign_Right)
+	[
 		SNew(STextBlock)
-		.Margin(FMargin(ColumnMargin, RowMargin))
+		.Margin(TitleMargin)
 		.ColorAndOpacity(TitleColor)
 		.Font(TitleFont)
-		.Justification(ETextJustify::Right)
-		.Text(FText::FromString(SingleDecimalFormat(SumTotalPutMB) + TEXT(" MB")))
-		];
+		.Text(FText::FromString(SingleDecimalFormat(SumTotalPutMB)))
+	];
 
 	return Panel;
 }
 
-#undef LOCTEXT_NAMESPACE
\ No newline at end of file
+#undef LOCTEXT_NAMESPACE
