Allow users to disable all or specific notifications from DDC system

diff --git a/Engine/Source/Editor/DerivedDataEditor/Private/DerivedDataInformation.cpp b/Engine/Source/Editor/DerivedDataEditor/Private/DerivedDataInformation.cpp
--- a/Engine/Source/Editor/DerivedDataEditor/Private/DerivedDataInformation.cpp
+++ b/Engine/Source/Editor/DerivedDataEditor/Private/DerivedDataInformation.cpp
@@ -40,7 +40,7 @@ double FDerivedDataInformation::GetCacheActivitySizeBytes(bool bGet, bool bLocal
 	for (const TSharedRef<const FDerivedDataCacheStatsNode>& Usage : GetCacheUsageStats())
 	{
 		if (Usage->IsLocal() != bLocal)
-		{
+	{
 			continue;
 		}
 
@@ -70,7 +70,7 @@ double FDerivedDataInformation::GetCacheActivityTimeSeconds(bool bGet, bool bLoc
 	for (const TSharedRef<const FDerivedDataCacheStatsNode>& Usage : GetCacheUsageStats())
 	{
 		if (Usage->IsLocal() != bLocal)
-		{
+	{
 			continue;
 		}
 
@@ -103,11 +103,11 @@ double FDerivedDataInformation::GetCacheActivityTimeSeconds(bool bGet, bool bLoc
 bool FDerivedDataInformation::GetHasRemoteCache()
 {
 	for (const TSharedRef<const FDerivedDataCacheStatsNode>& Usage : GetCacheUsageStats())
-	{
-		if (!Usage->IsLocal())
 		{
+		if (!Usage->IsLocal())
+	{
 			return true;
-		}
+	}
 	}
 	return false;
 }
@@ -115,11 +115,11 @@ bool FDerivedDataInformation::GetHasRemoteCache()
 bool FDerivedDataInformation::GetHasZenCache()
 {
 	for (const TSharedRef<const FDerivedDataCacheStatsNode>& Usage : GetCacheUsageStats())
-	{
-		if (Usage->GetCacheType().Equals(TEXT("Zen")))
 		{
+		if (Usage->GetCacheType().Equals(TEXT("Zen")))
+	{
 			return true;
-		}
+	}
 	}
 	return false;
 }
@@ -127,11 +127,11 @@ bool FDerivedDataInformation::GetHasZenCache()
 bool FDerivedDataInformation::GetHasHordeStorageCache()
 {
 	for (const TSharedRef<const FDerivedDataCacheStatsNode>& Usage : GetCacheUsageStats())
-	{
-		if (Usage->GetCacheType().Equals(TEXT("Horde Storage")))
 		{
+		if (Usage->GetCacheType().Equals(TEXT("Horde Storage")))
+	{
 			return true;
-		}
+	}
 	}
 	return false;
 }
@@ -140,7 +140,7 @@ void FDerivedDataInformation::UpdateRemoteCacheState()
 {
 	RemoteCacheState = ERemoteCacheState::Unavailable;
 
-	if (GetHasRemoteCache())
+	if ( GetHasRemoteCache() )
 	{
 		const double OldLastGetTime = LastGetTime;
 		const double OldLastPutTime = LastPutTime;
@@ -166,32 +166,34 @@ void FDerivedDataInformation::UpdateRemoteCacheState()
 
 	if (const UDDCProjectSettings* DDCProjectSettings = GetDefault<UDDCProjectSettings>(); DDCProjectSettings && DDCProjectSettings->EnableWarnings)
 	{
-		if (const UEditorSettings* EditorSettings = GetDefault<UEditorSettings>())
+		const UEditorSettings* EditorSettings = GetDefault<UEditorSettings>();
+
+		if (EditorSettings && EditorSettings->bEnableDDCNotifications)
 		{
-			if (DDCProjectSettings->RecommendEveryoneUseHordeStorage && !GetHasHordeStorageCache() && (FCString::Stricmp(GetDerivedDataCache()->GetGraphName(), TEXT("NoJupiter")) != 0))
+			if (DDCProjectSettings->RecommendEveryoneUseHordeStorage && EditorSettings->bNotifyUseHordeStorage && !GetHasHordeStorageCache() && (FCString::Stricmp(GetDerivedDataCache()->GetGraphName(), TEXT("NoJupiter"))!=0))
 			{
 				RemoteCacheState = ERemoteCacheState::Warning;
 				RemoteCacheWarningMessage = FText(LOCTEXT("HordeStorageWarning", "It is recommended that you use a DDC graph that supports Horde Storage. Please check any -ddc commandline overrides."));
 			}
-			else if (DDCProjectSettings->RecommendEveryoneSetupAGlobalLocalDDCPath && EditorSettings->GlobalLocalDDCPath.Path.IsEmpty())
+			else if (DDCProjectSettings->RecommendEveryoneSetupAGlobalLocalDDCPath && EditorSettings->bNotifySetupDDCPath && EditorSettings->GlobalLocalDDCPath.Path.IsEmpty())
 			{
 				RemoteCacheState = ERemoteCacheState::Warning;
 				RemoteCacheWarningMessage = FText(LOCTEXT("GlobalLocalDDCPathWarning", "It is recommended that you set up a valid Global Local DDC Path"));
 			}
-			else if (DDCProjectSettings->RecommendEveryoneSetupAGlobalSharedDDCPath && EditorSettings->GlobalSharedDDCPath.Path.IsEmpty())
+			else if (DDCProjectSettings->RecommendEveryoneSetupAGlobalSharedDDCPath && EditorSettings->bNotifySetupDDCPath && EditorSettings->GlobalSharedDDCPath.Path.IsEmpty())
 			{
 				RemoteCacheState = ERemoteCacheState::Warning;
 				RemoteCacheWarningMessage = FText(LOCTEXT("GlobalLocalDDCPathWarning", "It is recommended that you set up a valid Global Shared DDC Path"));
 			}
-			else if (DDCProjectSettings->RecommendEveryoneEnableS3DDC && !EditorSettings->bEnableS3DDC)
+			else if (DDCProjectSettings->RecommendEveryoneEnableS3DDC && EditorSettings->bNotifyEnableS3DD && !EditorSettings->bEnableS3DDC)
 			{
 				RemoteCacheState = ERemoteCacheState::Warning;
 				RemoteCacheWarningMessage = FText(LOCTEXT("AWSS3CacheEnabledWarning", "It is recommended that you enable the AWS S3 Cache"));
 			}
-			else if (DDCProjectSettings->RecommendEveryoneSetupAGlobalS3DDCPath && EditorSettings->GlobalS3DDCPath.Path.IsEmpty())
+			else if (DDCProjectSettings->RecommendEveryoneSetupAGlobalS3DDCPath && EditorSettings->bNotifySetupDDCPath && EditorSettings->GlobalS3DDCPath.Path.IsEmpty())
 			{
 				RemoteCacheState = ERemoteCacheState::Warning;
-				RemoteCacheWarningMessage = FText(LOCTEXT("S3GloblaLocalPathdWarning", "It is recommended that you set up a valid Global Local S3 DDC Path"));
+				RemoteCacheWarningMessage = FText(LOCTEXT("S3GloblaLocalPathWarning", "It is recommended that you set up a valid Global Local S3 DDC Path"));
 			}
 		}
 	}
diff --git a/Engine/Source/Editor/UnrealEd/Public/Settings/EditorSettings.h b/Engine/Source/Editor/UnrealEd/Public/Settings/EditorSettings.h
--- a/Engine/Source/Editor/UnrealEd/Public/Settings/EditorSettings.h
+++ b/Engine/Source/Editor/UnrealEd/Public/Settings/EditorSettings.h
@@ -80,6 +80,22 @@ class UNREALED_API UEditorSettings : public UObject
 	UPROPERTY(EditAnywhere, config, Category= DerivedDataCache, AdvancedDisplay, meta = (DisplayName = "Project Shared DDC Path", ConfigRestartRequired = true))
 	FDirectoryPath SharedDerivedDataCache;
 
+	/** Whether to enable any DDC Notifications */
+	UPROPERTY(EditAnywhere, config, Category = "Derived Data Cache Notifications", meta = (DisplayName = "Enable Notifcations", ConfigRestartRequired = false))
+	bool bEnableDDCNotifications = true;
+
+	/** Whether to enable the Horde Storage notification */
+	UPROPERTY(EditAnywhere, config, Category = "Derived Data Cache Notifications", meta = (DisplayName = "Notify Use Horde Storage", ConfigRestartRequired = false, EditCondition = "bEnableDDCNotifications"))
+	bool bNotifyUseHordeStorage = true;
+
+	/** Whether to enable the DDC path notification */
+	UPROPERTY(EditAnywhere, config, Category = "Derived Data Cache Notifications", meta = (DisplayName = "Notify Setup DDC Path", ConfigRestartRequired = false, EditCondition = "bEnableDDCNotifications"))
+	bool bNotifySetupDDCPath = true;
+
+	/** Whether to enable the DDC path notification */
+	UPROPERTY(EditAnywhere, config, Category = "Derived Data Cache Notifications", meta = (DisplayName = "Notify Enable S3 DDC", ConfigRestartRequired = false, EditCondition = "bEnableDDCNotifications"))
+	bool bNotifyEnableS3DD = true;
+
 	/** Whether to enable the S3 derived data cache backend */
 	UPROPERTY(EditAnywhere, config, Category="Derived Data Cache S3", meta = (DisplayName = "Enable AWS S3 Cache", ConfigRestartRequired = true))
 	bool bEnableS3DDC = true;
