Fix for Virtual Assets patload notification popping too fequently

diff --git a/Engine/Source/Editor/DerivedDataEditor/Private/SVirtualAssetsStatistics.cpp b/Engine/Source/Editor/DerivedDataEditor/Private/SVirtualAssetsStatistics.cpp
--- a/Engine/Source/Editor/DerivedDataEditor/Private/SVirtualAssetsStatistics.cpp
+++ b/Engine/Source/Editor/DerivedDataEditor/Private/SVirtualAssetsStatistics.cpp
@@ -91,19 +91,31 @@ void SVirtualAssetsStatisticsDialog::Construct(const FArguments& InArgs)
 		]
 	];
 
-	RegisterActiveTimer(0.5f, FWidgetActiveTimerDelegate::CreateSP(this, &SVirtualAssetsStatisticsDialog::UpdateGridPanels));
+	RegisterActiveTimer(0.25f, FWidgetActiveTimerDelegate::CreateSP(this, &SVirtualAssetsStatisticsDialog::UpdateGridPanels));
 }
 
 EActiveTimerReturnType SVirtualAssetsStatisticsDialog::UpdateGridPanels(double InCurrentTime, float InDeltaTime)
 {
 	(*GridSlot)
-	[
-		GetGridPanel()
-	];
+		[
+			GetGridPanel()
+		];
 
 	SlatePrepass(GetPrepassLayoutScaleMultiplier());
 
-	if ( NumPullRequests!=0 && PullRequestNotificationItem.IsValid()==false )
+	const float PullNotifactionTimeLimit=1.0f;
+
+	// Only show the pull notification if we have been pulling for more than a second..
+	if (NumPullRequests != 0)
+	{
+		PullNotificationTimer += InDeltaTime;
+	}
+	else
+	{
+		PullNotificationTimer = 0.0f;
+	}
+
+	if ( PullNotificationTimer>PullNotifactionTimeLimit && PullRequestNotificationItem.IsValid()==false )
 	{
 		// No existing notification or the existing one has finished
 		TPromise<TWeakPtr<SNotificationItem>> NotificationPromise;
diff --git a/Engine/Source/Editor/DerivedDataEditor/Private/SVirtualAssetsStatistics.h b/Engine/Source/Editor/DerivedDataEditor/Private/SVirtualAssetsStatistics.h
--- a/Engine/Source/Editor/DerivedDataEditor/Private/SVirtualAssetsStatistics.h
+++ b/Engine/Source/Editor/DerivedDataEditor/Private/SVirtualAssetsStatistics.h
@@ -38,4 +38,5 @@ private:
 
 	bool	IsPulling = false;
 	uint32	NumPullRequests = 0;
+	float	PullNotificationTimer = 0.0f;
 };
\ No newline at end of file
