Add a horizontal scrollbox to SVirtualAssetsStatisticsDialog

#rnx

diff --git a/Engine/Source/Editor/DerivedDataEditor/Private/SVirtualAssetsStatistics.cpp b/Engine/Source/Editor/DerivedDataEditor/Private/SVirtualAssetsStatistics.cpp
--- a/Engine/Source/Editor/DerivedDataEditor/Private/SVirtualAssetsStatistics.cpp
+++ b/Engine/Source/Editor/DerivedDataEditor/Private/SVirtualAssetsStatistics.cpp
@@ -1,16 +1,18 @@
 // Copyright Epic Games, Inc. All Rights Reserved.
 
 #include "SVirtualAssetsStatistics.h"
-#include "VirtualizationManager.h"
-#include "Math/UnitConversion.h"
+
+#include "Framework/Notifications/NotificationManager.h"
+#include "Internationalization/FastDecimalFormat.h"
 #include "Math/BasicMathExpressionEvaluator.h"
+#include "Math/UnitConversion.h"
 #include "Misc/ExpressionParser.h"
 #include "Styling/StyleColors.h"
+#include "VirtualizationManager.h"
 #include "Widgets/Layout/SGridPanel.h"
-#include "Widgets/Text/STextBlock.h"
+#include "Widgets/Layout/SScrollBox.h"
 #include "Widgets/SBoxPanel.h"
-#include "Internationalization/FastDecimalFormat.h"
-#include "Framework/Notifications/NotificationManager.h"
+#include "Widgets/Text/STextBlock.h"
 
 #define LOCTEXT_NAMESPACE "VirtualAssets"
 
@@ -87,7 +89,14 @@ void SVirtualAssetsStatisticsDialog::Construct(const FArguments& InArgs)
 		.Padding(0, 20, 0, 0)
 		.Expose(GridSlot)
 		[
-			GetGridPanel()
+			SAssignNew(ScrollBox, SScrollBox)
+			.Orientation(EOrientation::Orient_Horizontal)
+			.ScrollBarAlwaysVisible(false)
+
+			+ SScrollBox::Slot()
+			[
+				GetGridPanel()
+			]	
 		]
 	];
 
@@ -96,11 +105,12 @@ void SVirtualAssetsStatisticsDialog::Construct(const FArguments& InArgs)
 
 EActiveTimerReturnType SVirtualAssetsStatisticsDialog::UpdateGridPanels(double InCurrentTime, float InDeltaTime)
 {
-	(*GridSlot)
+	ScrollBox->ClearChildren();
+	ScrollBox->AddSlot()
 		[
 			GetGridPanel()
 		];
-
+		
 	SlatePrepass(GetPrepassLayoutScaleMultiplier());
 
 	const float PullNotifactionTimeLimit=1.0f;
diff --git a/Engine/Source/Editor/DerivedDataEditor/Private/SVirtualAssetsStatistics.h b/Engine/Source/Editor/DerivedDataEditor/Private/SVirtualAssetsStatistics.h
--- a/Engine/Source/Editor/DerivedDataEditor/Private/SVirtualAssetsStatistics.h
+++ b/Engine/Source/Editor/DerivedDataEditor/Private/SVirtualAssetsStatistics.h
@@ -36,6 +36,8 @@ private:
 
 	TSharedPtr<SNotificationItem> PullRequestNotificationItem;
 
+	TSharedPtr<class SScrollBox> ScrollBox;
+
 	bool	IsPulling = false;
 	uint32	NumPullRequests = 0;
 	float	PullNotificationTimer = 0.0f;
