The editor can now compile when ENABLE_COOK_STATS is set to 0

#rnx

- Fixed up a number of places that needed to be wrapped in #if ENABLE_COOK_STATS or #if OUTPUT_COOKTIMING
- The SDerivedDataCacheStatisticsDialog cannot display anything when cooking stats are disabled, so instead we show a message explaining why the dialog has no info

diff --git a/Engine/Source/Editor/DerivedDataEditor/Private/DerivedDataInformation.cpp b/Engine/Source/Editor/DerivedDataEditor/Private/DerivedDataInformation.cpp
--- a/Engine/Source/Editor/DerivedDataEditor/Private/DerivedDataInformation.cpp
+++ b/Engine/Source/Editor/DerivedDataEditor/Private/DerivedDataInformation.cpp
@@ -37,10 +37,11 @@ double FDerivedDataInformation::GetCacheActivitySizeBytes(bool bGet, bool bLocal
 {
 	int64 TotalBytes = 0;
 
+#if ENABLE_COOK_STATS
 	for (const TSharedRef<const FDerivedDataCacheStatsNode>& Usage : GetCacheUsageStats())
 	{
 		if (Usage->IsLocal() != bLocal)
-	{
+		{
 			continue;
 		}
 
@@ -58,6 +59,7 @@ double FDerivedDataInformation::GetCacheActivitySizeBytes(bool bGet, bool bLocal
 			}
 		}
 	}
+#endif // ENABLE_COOK_STATS
 
 	return TotalBytes;
 }
@@ -67,10 +69,11 @@ double FDerivedDataInformation::GetCacheActivityTimeSeconds(bool bGet, bool bLoc
 {
 	int64 TotalCycles = 0;
 
+#if ENABLE_COOK_STATS
 	for (const TSharedRef<const FDerivedDataCacheStatsNode>& Usage : GetCacheUsageStats())
 	{
 		if (Usage->IsLocal() != bLocal)
-	{
+		{
 			continue;
 		}
 
@@ -96,6 +99,7 @@ double FDerivedDataInformation::GetCacheActivityTimeSeconds(bool bGet, bool bLoc
 			}
 		}
 	}
+#endif // ENABLE_COOK_STATS
 
 	return (double)TotalCycles * FPlatformTime::GetSecondsPerCycle();
 }
diff --git a/Engine/Source/Editor/DerivedDataEditor/Private/SDerivedDataDialogs.cpp b/Engine/Source/Editor/DerivedDataEditor/Private/SDerivedDataDialogs.cpp
--- a/Engine/Source/Editor/DerivedDataEditor/Private/SDerivedDataDialogs.cpp
+++ b/Engine/Source/Editor/DerivedDataEditor/Private/SDerivedDataDialogs.cpp
@@ -523,17 +523,20 @@ TSharedRef<SWidget> SDerivedDataCacheStatisticsDialog::GetGridPanel()
 	TSharedRef<SGridPanel> Panel =
 		SNew(SGridPanel);
 
-	int32 Row = 0;
-
 	const float RowMargin = 0.0f;
 	const float ColumnMargin = 10.0f;
-	const FMargin TitleMargin(0.0f, 10.0f, ColumnMargin, 10.0f);
-	const FMargin TitleMarginFirstColumn(ColumnMargin, 10.0f);
 	const FSlateColor TitleColor = FStyleColors::AccentWhite;
 	const FSlateFontInfo TitleFont = FCoreStyle::GetDefaultFontStyle("Bold", 10);
-	const FMargin DefaultMargin(0.0f, RowMargin, ColumnMargin, RowMargin);
 	const FMargin DefaultMarginFirstColumn(ColumnMargin, RowMargin);
 
+#if ENABLE_COOK_STATS
+
+	int32 Row = 0;
+
+	const FMargin TitleMargin(0.0f, 10.0f, ColumnMargin, 10.0f);
+	const FMargin TitleMarginFirstColumn(ColumnMargin, 10.0f);
+	const FMargin DefaultMargin(0.0f, RowMargin, ColumnMargin, RowMargin);
+	
 	Panel->AddSlot(0, Row)
 	[
 		SNew(STextBlock)
@@ -702,6 +705,19 @@ TSharedRef<SWidget> SDerivedDataCacheStatisticsDialog::GetGridPanel()
 		.Text(FText::FromString(SingleDecimalFormat(SumTotalPutMB)))
 	];
 
+#else
+	Panel->AddSlot(0, 0)
+	[
+		SNew(STextBlock)
+		.Margin(DefaultMarginFirstColumn)
+		.ColorAndOpacity(TitleColor)
+		.Font(TitleFont)
+		.Justification(ETextJustify::Center)
+		.Text(LOCTEXT("Disabled", "Cooking Stats Are Disabled For This Project"))
+	];
+
+#endif // ENABLE_COOK_STATS
+
 	return Panel;
 }
 
diff --git a/Engine/Source/Editor/UnrealEd/Private/Commandlets/CookCommandlet.cpp b/Engine/Source/Editor/UnrealEd/Private/Commandlets/CookCommandlet.cpp
--- a/Engine/Source/Editor/UnrealEd/Private/Commandlets/CookCommandlet.cpp
+++ b/Engine/Source/Editor/UnrealEd/Private/Commandlets/CookCommandlet.cpp
@@ -702,7 +702,9 @@ int32 UCookCommandlet::Main(const FString& CmdLineParams)
 
 bool UCookCommandlet::CookByTheBook( const TArray<ITargetPlatform*>& Platforms)
 {
+#if OUTPUT_COOKTIMING
 	TRACE_CPUPROFILER_EVENT_SCOPE_ON_CHANNEL(CookByTheBook, CookChannel);
+#endif // OUTPUT_COOKTIMING
 
 	COOK_STAT(FScopedDurationTimer CookByTheBookTimer(DetailedCookStats::CookByTheBookTimeSec));
 	UCookOnTheFlyServer *CookOnTheFlyServer = NewObject<UCookOnTheFlyServer>();
diff --git a/Engine/Source/Editor/UnrealEd/Private/CookOnTheFlyServer.cpp b/Engine/Source/Editor/UnrealEd/Private/CookOnTheFlyServer.cpp
--- a/Engine/Source/Editor/UnrealEd/Private/CookOnTheFlyServer.cpp
+++ b/Engine/Source/Editor/UnrealEd/Private/CookOnTheFlyServer.cpp
@@ -4068,10 +4068,14 @@ private:
 	bool bLinkerDiffEnabled = false;
 };
 
+#if OUTPUT_COOKTIMING
+
 UE_TRACE_EVENT_BEGIN(UE_CUSTOM_COOKTIMER_LOG, SaveCookedPackage, NoSync)
 	UE_TRACE_EVENT_FIELD(UE::Trace::WideString, PackageName)
 UE_TRACE_EVENT_END()
 
+#endif //OUTPUT_COOKTIMING
+
 void UCookOnTheFlyServer::SaveCookedPackage(UE::Cook::FSaveCookedPackageContext& Context)
 {
 	using namespace UE::Cook;
diff --git a/Engine/Source/Editor/UnrealEd/Private/Cooker/CookProfiling.h b/Engine/Source/Editor/UnrealEd/Private/Cooker/CookProfiling.h
--- a/Engine/Source/Editor/UnrealEd/Private/Cooker/CookProfiling.h
+++ b/Engine/Source/Editor/UnrealEd/Private/Cooker/CookProfiling.h
@@ -76,10 +76,10 @@ UE_TRACE_CHANNEL_EXTERN(CookChannel)
 #define UE_SCOPED_CUSTOM_COOKTIMER_AND_DURATION(name, durationStorage)
 #define UE_SCOPED_HIERARCHICAL_CUSTOM_COOKTIMER(name)
 #define UE_SCOPED_HIERARCHICAL_CUSTOM_COOKTIMER_AND_DURATION(name, durationStorage)
-#define UE_ADD_CUSTOM_COOKTIMER_META(key, value)
+#define UE_ADD_CUSTOM_COOKTIMER_META(name, key, value)
 
-void OutputHierarchyTimers() {}
-void ClearHierarchyTimers() {}
+inline void OutputHierarchyTimers() {}
+inline void ClearHierarchyTimers() {}
 #endif
 
 #if ENABLE_COOK_STATS
