Move the 'Virtual Assets' menu entry from 'Tools' to 'Tools -> Audit'

- The UX team think s that the audit submenu is a better fit for SVirtualAssetsStatisticsDialog than the top level.
- Given that the dialogs main use is to audit the amount of data that the user has moved around this makes sense.

diff --git a/Engine/Source/Editor/DerivedDataEditor/Private/DerivedDataEditorModule.cpp b/Engine/Source/Editor/DerivedDataEditor/Private/DerivedDataEditorModule.cpp
--- a/Engine/Source/Editor/DerivedDataEditor/Private/DerivedDataEditorModule.cpp
+++ b/Engine/Source/Editor/DerivedDataEditor/Private/DerivedDataEditorModule.cpp
@@ -50,7 +50,7 @@ void FDerivedDataEditorModule::StartupModule()
 	FGlobalTabmanager::Get()->RegisterNomadTabSpawner(VirtualAssetsStatisticsTabName, FOnSpawnTab::CreateRaw(this, &FDerivedDataEditorModule::CreateVirtualAssetsStatisticsTab))
 		.SetDisplayName(LOCTEXT("VirtualAssetsStatisticsTabTitle", "Virtual Assets"))
 		.SetTooltipText(LOCTEXT("VirtualAssetsStatisticsTabToolTipText", "Virtual Assets  Statistics"))
-		.SetGroup(WorkspaceMenu::GetMenuStructure().GetToolsCategory())
+		.SetGroup(WorkspaceMenu::GetMenuStructure().GetDeveloperToolsAuditCategory())
 		.SetIcon(CacheStatisticsIcon);
 
 
