Implemented Editor notifcations when BulkData payloads cannot be reached through any Virtual Aseet backend.

diff --git a/Engine/Source/Developer/Virtualization/Private/VirtualizationManager.cpp b/Engine/Source/Developer/Virtualization/Private/VirtualizationManager.cpp
--- a/Engine/Source/Developer/Virtualization/Private/VirtualizationManager.cpp
+++ b/Engine/Source/Developer/Virtualization/Private/VirtualizationManager.cpp
@@ -4,6 +4,7 @@
 
 #include "HAL/PlatformTime.h"
 #include "IVirtualizationBackend.h"
+#include "Logging/MessageLog.h"
 #include "Misc/CommandLine.h"
 #include "Misc/ConfigCacheIni.h"
 #include "Misc/CoreDelegates.h"
@@ -16,6 +17,7 @@
 #include "ProfilingDebugging/CookStats.h"
 #include "VirtualizationFilterSettings.h"
 
+#define LOCTEXT_NAMESPACE "Virtualization"
 
 namespace UE::Virtualization
 {
@@ -1331,3 +1333,5 @@ bool FVirtualizationManager::ShouldVirtualizeAsDefault() const
 }
 
 } // namespace UE::Virtualization
+
+#undef LOCTEXT_NAMESPACE
\ No newline at end of file
diff --git a/Engine/Source/Editor/DerivedDataEditor/Private/SVirtualAssetsStatistics.cpp b/Engine/Source/Editor/DerivedDataEditor/Private/SVirtualAssetsStatistics.cpp
--- a/Engine/Source/Editor/DerivedDataEditor/Private/SVirtualAssetsStatistics.cpp
+++ b/Engine/Source/Editor/DerivedDataEditor/Private/SVirtualAssetsStatistics.cpp
@@ -62,15 +62,7 @@ void SVirtualAssetsStatisticsDialog::OnNotificationEvent(IVirtualizationSystem::
 
 		case IVirtualizationSystem::ENotification::PullFailedNotification:
 		{
-			FNotificationInfo* Info = new FNotificationInfo(LOCTEXT("PayloadSyncFail", "Failed To Sync Asset Payload"));
-
-			Info->bUseSuccessFailIcons = true;
-			Info->bFireAndForget = true;
-			Info->FadeOutDuration = 1.0f;
-			Info->ExpireDuration = 4.0f;
-
-			FSlateNotificationManager::Get().QueueNotification(Info);
-
+			NumPullRequestFailures++;
 			break;
 		}
 
@@ -128,10 +120,10 @@ EActiveTimerReturnType SVirtualAssetsStatisticsDialog::UpdateGridPanels(double I
 	if ( PullNotificationTimer>PullNotifactionTimeLimit && PullRequestNotificationItem.IsValid()==false )
 	{
 		// No existing notification or the existing one has finished
-		TPromise<TWeakPtr<SNotificationItem>> NotificationPromise;
-
-		FNotificationInfo Info(LOCTEXT("PayloadSyncInProgress", "Syncing Asset Payloads"));
+		FNotificationInfo Info(LOCTEXT("PayloadSyncNotifcation", "Syncing Asset Payloads"));
 		Info.bFireAndForget = false;
+		Info.bUseLargeFont = false;
+		Info.bUseThrobber = false;
 		Info.FadeOutDuration = 0.5f;
 		Info.ExpireDuration = 0.0f;
 
@@ -139,10 +131,27 @@ EActiveTimerReturnType SVirtualAssetsStatisticsDialog::UpdateGridPanels(double I
 
 		if (PullRequestNotificationItem.IsValid())
 		{
-			NotificationPromise.SetValue(PullRequestNotificationItem);
 			PullRequestNotificationItem->SetCompletionState(SNotificationItem::CS_Pending);
 		}
 	}
+
+	if ( NumPullRequestFailures>0 && PullRequestFailedNotificationItem.IsValid()==false )
+	{
+		// No existing notification or the existing one has finished
+		FNotificationInfo Info(LOCTEXT("PayloadFailedNotifcation", "Failed to sync some Virtual Asset payloads from available backends.\nSome assets may no longer be usable.."));	
+		Info.bFireAndForget = false;
+		Info.bUseLargeFont = false;
+		Info.bUseThrobber = false;
+		Info.FadeOutDuration = 0.5f;
+		Info.ExpireDuration = 0.0f;
+		Info.Image = FEditorStyle::GetBrush(TEXT("MessageLog.Warning"));
+		Info.ButtonDetails.Add(FNotificationButtonInfo(LOCTEXT("PullFailedIgnore", "Ignore"), LOCTEXT("PullFailedIgnoreToolTip", "Ignore future warnings"), FSimpleDelegate::CreateSP(this, &SVirtualAssetsStatisticsDialog::OnWarningReasonIgnore), SNotificationItem::CS_None));
+		Info.ButtonDetails.Add(FNotificationButtonInfo(LOCTEXT("PullFailedOK", "Ok"), LOCTEXT("PullFailedOkToolTip", "Notify future warnings"), FSimpleDelegate::CreateSP(this, &SVirtualAssetsStatisticsDialog::OnWarningReasonOk), SNotificationItem::CS_None));
+		Info.HyperlinkText = LOCTEXT("PullFailed_ShowLog", "Show Message Log");
+		Info.Hyperlink = FSimpleDelegate::CreateStatic([]() { FMessageLog("LogVirtualization").Open(EMessageSeverity::Warning, true); });
+
+		PullRequestFailedNotificationItem = FSlateNotificationManager::Get().AddNotification(Info);
+	}
 	
 	if ( NumPullRequests==0 && PullRequestNotificationItem.IsValid()==true )
 	{
@@ -154,6 +163,24 @@ EActiveTimerReturnType SVirtualAssetsStatisticsDialog::UpdateGridPanels(double I
 	return EActiveTimerReturnType::Continue;
 }
 
+void SVirtualAssetsStatisticsDialog::OnWarningReasonOk()
+{
+	if (PullRequestFailedNotificationItem.IsValid() == true)
+	{
+		PullRequestFailedNotificationItem->ExpireAndFadeout();
+		PullRequestFailedNotificationItem.Reset();
+		NumPullRequestFailures = 0;
+	}
+}
+
+void SVirtualAssetsStatisticsDialog::OnWarningReasonIgnore()
+{
+	if (PullRequestFailedNotificationItem.IsValid() == true)
+	{
+		PullRequestFailedNotificationItem->ExpireAndFadeout();
+	}
+}
+
 TSharedRef<SWidget> SVirtualAssetsStatisticsDialog::GetGridPanel()
 {
 	IVirtualizationSystem& System = IVirtualizationSystem::Get();
diff --git a/Engine/Source/Editor/DerivedDataEditor/Private/SVirtualAssetsStatistics.h b/Engine/Source/Editor/DerivedDataEditor/Private/SVirtualAssetsStatistics.h
--- a/Engine/Source/Editor/DerivedDataEditor/Private/SVirtualAssetsStatistics.h
+++ b/Engine/Source/Editor/DerivedDataEditor/Private/SVirtualAssetsStatistics.h
@@ -29,16 +29,21 @@ private:
 	FText GetNotificationText() const;
 
 	void OnNotificationEvent(UE::Virtualization::IVirtualizationSystem::ENotification Notification, const FIoHash& PayloadId);
+	void OnWarningReasonOk();
+	void OnWarningReasonIgnore();
 
 	SVerticalBox::FSlot* GridSlot = nullptr;
 
 	FCriticalSection NotificationCS;
 
 	TSharedPtr<SNotificationItem> PullRequestNotificationItem;
+	TSharedPtr<SNotificationItem> PullRequestFailedNotificationItem;
 
 	TSharedPtr<class SScrollBox> ScrollBox;
 
 	bool	IsPulling = false;
-	uint32	NumPullRequests = 0;
 	float	PullNotificationTimer = 0.0f;
-};
\ No newline at end of file
+	uint32	NumPullRequests = 0;
+	uint32	NumPullRequestFailures = 0;
+};
+
