Non-Unity Fix

diff --git a/Engine/Source/Editor/DerivedDataEditor/Private/SVirtualAssetsStatistics.cpp b/Engine/Source/Editor/DerivedDataEditor/Private/SVirtualAssetsStatistics.cpp
--- a/Engine/Source/Editor/DerivedDataEditor/Private/SVirtualAssetsStatistics.cpp
+++ b/Engine/Source/Editor/DerivedDataEditor/Private/SVirtualAssetsStatistics.cpp
@@ -7,7 +7,9 @@
 #include "Math/BasicMathExpressionEvaluator.h"
 #include "Math/UnitConversion.h"
 #include "Misc/ExpressionParser.h"
+#include "Logging/MessageLog.h"
 #include "Styling/StyleColors.h"
+#include "EditorStyleSet.h"
 #include "VirtualizationManager.h"
 #include "Widgets/Layout/SGridPanel.h"
 #include "Widgets/Layout/SScrollBox.h"
