Removed Horde DDC Backend warning notification when Zen Backend is present

diff --git a/Engine/Source/Editor/DerivedDataEditor/Private/DerivedDataInformation.cpp b/Engine/Source/Editor/DerivedDataEditor/Private/DerivedDataInformation.cpp
--- a/Engine/Source/Editor/DerivedDataEditor/Private/DerivedDataInformation.cpp
+++ b/Engine/Source/Editor/DerivedDataEditor/Private/DerivedDataInformation.cpp
@@ -174,7 +174,7 @@ void FDerivedDataInformation::UpdateRemoteCacheState()
 
 		if (EditorSettings && EditorSettings->bEnableDDCNotifications)
 		{
-			if (DDCProjectSettings->RecommendEveryoneUseHordeStorage && EditorSettings->bNotifyUseHordeStorage && !GetHasHordeStorageCache() && (FCString::Stricmp(GetDerivedDataCache()->GetGraphName(), TEXT("NoJupiter"))!=0))
+			if (DDCProjectSettings->RecommendEveryoneUseHordeStorage && EditorSettings->bNotifyUseHordeStorage && !GetHasZenCache() && !GetHasHordeStorageCache() && (FCString::Stricmp(GetDerivedDataCache()->GetGraphName(), TEXT("NoJupiter")) != 0))
 			{
 				RemoteCacheState = ERemoteCacheState::Warning;
 				RemoteCacheWarningMessage = FText(LOCTEXT("HordeStorageWarning", "It is recommended that you use a DDC graph that supports Horde Storage. Please check any -ddc commandline overrides."));
