Add speed stats to the DDC Cache Statistics dialog.

diff --git a/Engine/Source/Editor/DerivedDataEditor/Private/SDerivedDataDialogs.cpp b/Engine/Source/Editor/DerivedDataEditor/Private/SDerivedDataDialogs.cpp
--- a/Engine/Source/Editor/DerivedDataEditor/Private/SDerivedDataDialogs.cpp
+++ b/Engine/Source/Editor/DerivedDataEditor/Private/SDerivedDataDialogs.cpp
@@ -574,6 +574,8 @@ TSharedRef<SWidget> SDerivedDataCacheStatisticsDialog::GetGridPanel()
 		.ColorAndOpacity(TitleColor)
 		.Font(TitleFont)
 		.Text(LOCTEXT("Read", "Read (MiB)"))
+		.AutoWrapText(true)
+		.WrapTextAt(66.0f)
 	];
 
 	Panel->AddSlot(4, Row)
@@ -584,15 +586,53 @@ TSharedRef<SWidget> SDerivedDataCacheStatisticsDialog::GetGridPanel()
 		.ColorAndOpacity(TitleColor)
 		.Font(TitleFont)
 		.Text(LOCTEXT("Write", "Write (MiB)"))
+		.AutoWrapText(true)
+		.WrapTextAt(66.0f)
 	];
 
 	Panel->AddSlot(5, Row)
+	.HAlign(HAlign_Right)
 	[
 		SNew(STextBlock)
 		.Margin(TitleMargin)
 		.ColorAndOpacity(TitleColor)
 		.Font(TitleFont)
-		.Text(LOCTEXT("Details", "Details"))
+		.Text(LOCTEXT("Latency", "Latency (ms)"))
+		.AutoWrapText(true)
+		.WrapTextAt(66.0f)
+	];
+
+	Panel->AddSlot(6, Row)
+	.HAlign(HAlign_Right)
+	[
+		SNew(STextBlock)
+		.Margin(TitleMargin)
+		.ColorAndOpacity(TitleColor)
+		.Font(TitleFont)
+		.Text(LOCTEXT("Read Speed", "Read Speed (MiB/s)"))
+		.AutoWrapText(true)
+		.WrapTextAt(66.0f)
+	];
+
+	Panel->AddSlot(7, Row)
+	.HAlign(HAlign_Right)
+	[
+		SNew(STextBlock)
+		.Margin(TitleMargin)
+		.ColorAndOpacity(TitleColor)
+		.Font(TitleFont)
+		.Text(LOCTEXT("Write Speed", "Write Speed (MiB/s)"))
+		.AutoWrapText(true)
+		.WrapTextAt(66.0f)
+	];
+
+	Panel->AddSlot(8, Row)
+	[
+		SNew(STextBlock)
+		.Margin(TitleMargin)
+		.ColorAndOpacity(TitleColor)
+		.Font(TitleFont)
+		.Text(LOCTEXT("Name", "Name"))
 	];
 
 	Row++;
@@ -666,7 +706,40 @@ TSharedRef<SWidget> SDerivedDataCacheStatisticsDialog::GetGridPanel()
 			.Text(FText::FromString(SingleDecimalFormat(TotalPutMB)))
 		];
 
-		Panel->AddSlot(5, Row)
+		if (Node->SpeedStats.LatencyMS)
+		{
+			Panel->AddSlot(5, Row)
+			.HAlign(HAlign_Right)
+			[
+				SNew(STextBlock)
+				.Margin(DefaultMargin)
+				.Text(FText::FromString(SingleDecimalFormat(Node->SpeedStats.LatencyMS)))
+			];
+		}
+
+		if (Node->SpeedStats.ReadSpeedMBs)
+		{
+			Panel->AddSlot(6, Row)
+			.HAlign(HAlign_Right)
+			[
+				SNew(STextBlock)
+				.Margin(DefaultMargin)
+				.Text(FText::FromString(SingleDecimalFormat(Node->SpeedStats.ReadSpeedMBs)))
+			];
+		}
+
+		if (Node->SpeedStats.WriteSpeedMBs)
+		{
+			Panel->AddSlot(7, Row)
+			.HAlign(HAlign_Right)
+			[
+				SNew(STextBlock)
+				.Margin(DefaultMargin)
+				.Text(FText::FromString(SingleDecimalFormat(Node->SpeedStats.WriteSpeedMBs)))
+			];
+		}
+
+		Panel->AddSlot(8, Row)
 		[
 			SNew(STextBlock)
 			.Margin(DefaultMargin)
