Add an opt in entry to the asset context menu that allows the user to re-hydrate a virtualized package.

#rnx

### SVirtualAssetsStatistics
- We do not currently have a virtualization specific editor module and have been piggybacking onto the DDC editor module, so although this new code has nothing to do with the VA statistics panel it seems easier to keep the code in a single file and then split it later when we move it to a virtualization editor module.
- At the moment if the files being hydrated are not checked out then the process will fail. Support for auto checkout will be added in a future submit.

### EditorExperimentalSettings
- By default no change will be made to the context menu. The user must edit the experimental editor settings and opt in to be able to see it.
- This is because the feature should not really be used in day to day workflows and is provided to either fix problems or to make it easier for people developing the virtualization system.

### PackageRehydrationProcess
- Fixed some typos in error messages
- Being unable to write to the package is considered an error rather than a warning. It was demoted to a warning when virtualizing to try and reduce the number of blocked submits but we do not have this problem with rehydration.
- Added code to reset the loader of a package if it happens to be already loaded in the editor now that we can invoke it from the editor.

diff --git a/Engine/Source/Developer/Virtualization/Private/PackageRehydrationProcess.cpp b/Engine/Source/Developer/Virtualization/Private/PackageRehydrationProcess.cpp
--- a/Engine/Source/Developer/Virtualization/Private/PackageRehydrationProcess.cpp
+++ b/Engine/Source/Developer/Virtualization/Private/PackageRehydrationProcess.cpp
@@ -6,8 +6,10 @@
 #include "Internationalization/Internationalization.h"
 #include "Misc/PackageName.h"
 #include "Misc/PackagePath.h"
+#include "Misc/ScopedSlowTask.h"
 #include "PackageUtils.h"
 #include "UObject/Linker.h"
+#include "UObject/Package.h"
 #include "UObject/PackageTrailer.h"
 #include "UObject/UObjectGlobals.h"
 #include "Virtualization/VirtualizationSystem.h"
@@ -28,6 +30,9 @@ void RehydratePackages(const TArray<FString>& Packages, TArray<FText>& OutErrors
 		return;
 	}
 
+	FScopedSlowTask Progress(1.0f, LOCTEXT("VAHydration_Task", "Re-hydrating Assets..."));
+	Progress.MakeDialog();
+
 	for (const FString& FilePath : Packages)
 	{
 		if (!FPackageName::IsPackageFilename(FilePath))
@@ -50,7 +55,7 @@ void RehydratePackages(const TArray<FString>& Packages, TArray<FText>& OutErrors
 		TUniquePtr<FArchive> PackageAr(IFileManager::Get().CreateFileReader(*FilePath));
 		if (!PackageAr.IsValid())
 		{
-			FText Message = FText::Format(LOCTEXT("VAHydration_REadFailed", "Unable to open the package '{0}' for reading"),
+			FText Message = FText::Format(LOCTEXT("VAHydration_ReadFailed", "Unable to open the package '{0}' for reading"),
 				FText::FromString(FilePath));
 			OutErrors.Add(Message);
 
@@ -96,7 +101,16 @@ void RehydratePackages(const TArray<FString>& Packages, TArray<FText>& OutErrors
 		{
 			const FString NewPackagePath = DuplicatePackageWithNewTrailer(FilePath, Trailer, Builder, OutErrors);
 
-			// TODO: Reset loaders here if we ever expose this path to be run from the editor
+			FString PackageName;
+			if (FPackageName::TryConvertFilenameToLongPackageName(FilePath, PackageName))
+			{
+				UPackage* Package = FindObjectFast<UPackage>(nullptr, *PackageName);
+				if (Package != nullptr)
+				{
+					UE_LOG(LogVirtualization, Verbose, TEXT("Detaching '%s' from disk so that it can be rehydrated"), *FilePath);
+					ResetLoadersForSave(Package, *FilePath);
+				}
+			}
 
 			if (CanWriteToFile(FilePath))                     
 			{
@@ -111,17 +125,12 @@ void RehydratePackages(const TArray<FString>& Packages, TArray<FText>& OutErrors
 			}
 			else
 			{
-				// Technically the package could have local payloads that won't be virtualized due to filtering or min payload sizes and so the
-				// following warning is misleading. This will be solved if we move that evaluation to the point of saving a package.
-				// If not then we probably need to extend QueryPayloadStatuses to test filtering etc as well, then check for potential package
-				// modification after that.
-				// Long term, the stand alone tool should be able to request the UnrealEditor relinquish the lock on the package file so this becomes 
-				// less of a problem.
 				FText Message = FText::Format(
-					LOCTEXT("VAHydration_PackageLocked", "The package file '{0}' has virtualized payloads but is locked for modification and cannot be hydrated, this package will be skipped!"),
+					LOCTEXT("VAHydration_PackageLocked", "The package file '{0}' has virtualized payloads but is locked for modification and cannot be hydrated"),
 					FText::FromString(FilePath));
 
-				UE_LOG(LogVirtualization, Warning, TEXT("%s"), *Message.ToString());
+				OutErrors.Add(Message);
+				return;
 			}
 		}
 	}
diff --git a/Engine/Source/Editor/DerivedDataEditor/DerivedDataEditor.Build.cs b/Engine/Source/Editor/DerivedDataEditor/DerivedDataEditor.Build.cs
--- a/Engine/Source/Editor/DerivedDataEditor/DerivedDataEditor.Build.cs
+++ b/Engine/Source/Editor/DerivedDataEditor/DerivedDataEditor.Build.cs
@@ -27,7 +27,8 @@ public class DerivedDataEditor : ModuleRules
 				"WorkspaceMenuStructure",
 				"MessageLog",
 				"ToolWidgets",
-				"Zen"
+				"Zen",
+				"ContentBrowserData",
 			});
 
 		DynamicallyLoadedModuleNames.AddRange(
diff --git a/Engine/Source/Editor/DerivedDataEditor/Private/SVirtualAssetsStatistics.cpp b/Engine/Source/Editor/DerivedDataEditor/Private/SVirtualAssetsStatistics.cpp
--- a/Engine/Source/Editor/DerivedDataEditor/Private/SVirtualAssetsStatistics.cpp
+++ b/Engine/Source/Editor/DerivedDataEditor/Private/SVirtualAssetsStatistics.cpp
@@ -2,19 +2,23 @@
 
 #include "SVirtualAssetsStatistics.h"
 
+#include "ContentBrowserDataMenuContexts.h"
+#include "ContentBrowserDataSource.h"
 #include "Framework/Notifications/NotificationManager.h"
 #include "Internationalization/FastDecimalFormat.h"
+#include "Logging/MessageLog.h"
 #include "Math/BasicMathExpressionEvaluator.h"
 #include "Math/UnitConversion.h"
 #include "Misc/ExpressionParser.h"
-#include "Logging/MessageLog.h"
+#include "Settings/EditorExperimentalSettings.h"
+#include "Styling/AppStyle.h"
 #include "Styling/StyleColors.h"
+#include "ToolMenus.h"
 #include "VirtualizationManager.h"
 #include "Widgets/Layout/SGridPanel.h"
 #include "Widgets/Layout/SScrollBox.h"
 #include "Widgets/SBoxPanel.h"
 #include "Widgets/Text/STextBlock.h"
-#include "Styling/AppStyle.h"
 
 #define LOCTEXT_NAMESPACE "VirtualAssets"
 
@@ -22,12 +26,132 @@ extern FString SingleDecimalFormat(double Value);
 
 using namespace UE::Virtualization;
 
+// This namespace contains code for adding the re-hydration options to the editor context menu.
+// For now it is easier to just put all virtualization editor code in this .cpp until we make
+// a virtualization specific editor module like 'DerivedDataEditor'
+namespace UE::Virtualization::Rehydration
+{
+
+/** 
+ * Attempt to rehydrate the provided files although this currently assumes that the user has
+ * checked the files out of source control before hand. 
+ */
+void RehydratePackages(TArray<FString> SelectedFiles)
+{
+	
+
+	TArray<FText> Errors;
+	if (IVirtualizationSystem::Get().TryRehydratePackages(SelectedFiles, Errors))
+	{
+		// TODO: At some point ::TryRehydratePackages will return more detail info about the process
+		// when it does we should make a better job at logging it.
+	
+		const FText Message = LOCTEXT("RehydrationSucccess", "Files were successfully re-hydrated");
+		
+		FNotificationInfo Info(Message);
+		Info.bFireAndForget = true;
+		Info.ExpireDuration = 2.0f;
+
+		FSlateNotificationManager::Get().AddNotification(Info);
+	}
+	else
+	{
+		const bool bForceNotification = true;
+
+		FMessageLog Log("LogVirtualization");
+
+		for (const FText& Error : Errors)
+		{
+			Log.Error(Error);
+		}
+
+		Log.Notify(LOCTEXT("RehydrationFailed", "Failed to rehydrate packages, see the message log for more info"), EMessageSeverity::Info, bForceNotification);
+	}
+}
+
+void BrowserItemsToFilePaths(const TArray<FContentBrowserItem>& Source, TArray<FString>& OutResult)
+{
+	for (const FContentBrowserItem& SelectedItem : Source)
+	{
+		const FContentBrowserItemData* SelectedItemData = SelectedItem.GetPrimaryInternalItem();
+		if (SelectedItemData == nullptr)
+		{
+			continue;
+		}
+
+		if (!SelectedItemData->IsFile())
+		{
+			continue;
+		}
+
+		UContentBrowserDataSource* DataSource = SelectedItemData->GetOwnerDataSource();
+		if (DataSource == nullptr)
+		{
+			continue;
+		}
+			
+		FString Path;
+		if (DataSource->GetItemPhysicalPath(*SelectedItemData, Path))
+		{
+			OutResult.Add(FPaths::ConvertRelativePathToFull(Path));
+		}
+	}
+}
+
+void AddContextMenuEntry()
+{
+	if (UToolMenu* Menu = UToolMenus::Get()->ExtendMenu("ContentBrowser.AssetContextMenu.AssetActionsSubMenu"))
+	{
+		Menu->AddDynamicSection("VirtualizedAssetsDynamic", FNewToolMenuDelegate::CreateLambda([](UToolMenu* Menu)
+		{
+			const UContentBrowserDataMenuContext_FileMenu* Context = Menu->FindContext<UContentBrowserDataMenuContext_FileMenu>();
+
+			if (Context == nullptr)
+			{
+				return;
+			}
+
+			if (Context->bCanBeModified == false)
+			{
+				return;
+			}
+
+			if (GetDefault<UEditorExperimentalSettings>()->bVirtualizedAssetRehydration == false)
+			{
+				return;
+			}
+
+			TArray<FString> SelectedFiles;
+			BrowserItemsToFilePaths(Context->SelectedItems, SelectedFiles);
+
+			if (SelectedFiles.IsEmpty())
+			{
+				return;
+			}
+
+			FToolMenuSection& Section = Menu->AddSection("VirtualizedAssets", LOCTEXT("VirtualizedAssetsHeading", "Virtualized Assets"));
+			{
+				Section.AddMenuEntry(
+					"RehydrateAsset",
+					LOCTEXT("RehydrateAsset", "Rehydrate Asset"),
+					LOCTEXT("RehydrateAssetTooltip", "Pulls the assets virtualized payloads and stores them in the package file once more"),
+					FSlateIcon(),
+					FUIAction(FExecuteAction::CreateStatic(&RehydratePackages, SelectedFiles))
+				);
+			}
+		}));
+	}
+}
+
+} // UE::Virtualization::Rehydration
+
 SVirtualAssetsStatisticsDialog::SVirtualAssetsStatisticsDialog()
 {
 	// Register our VA notification delegate with the event
 	IVirtualizationSystem& System = IVirtualizationSystem::Get();
 	System.GetNotificationEvent().AddRaw(this, &SVirtualAssetsStatisticsDialog::OnNotificationEvent);
 
+	UE::Virtualization::Rehydration::AddContextMenuEntry();
 }
 
 SVirtualAssetsStatisticsDialog::~SVirtualAssetsStatisticsDialog()
diff --git a/Engine/Source/Editor/UnrealEd/Classes/Settings/EditorExperimentalSettings.h b/Engine/Source/Editor/UnrealEd/Classes/Settings/EditorExperimentalSettings.h
--- a/Engine/Source/Editor/UnrealEd/Classes/Settings/EditorExperimentalSettings.h
+++ b/Engine/Source/Editor/UnrealEd/Classes/Settings/EditorExperimentalSettings.h
@@ -134,6 +134,10 @@ public:
 	UPROPERTY(EditAnywhere, config, Category = Core)
 	bool bTextAssetFormatSupport;
 
+	/** Enables in-editor support for rehydrating virtualized assets */
+	UPROPERTY(EditAnywhere, config, Category = Core)
+	bool bVirtualizedAssetRehydration;
+
 	/** When creating new Material Layers and Material Layer Blends, set up example graphs. */
 	UPROPERTY(EditAnywhere, config, Category = Materials)
 	bool bExampleLayersAndBlends;
