Move using declaration out of the global scope as per the coding standards.

#rnx

diff --git a/Engine/Source/Editor/DerivedDataEditor/Private/SVirtualAssetsStatistics.cpp b/Engine/Source/Editor/DerivedDataEditor/Private/SVirtualAssetsStatistics.cpp
--- a/Engine/Source/Editor/DerivedDataEditor/Private/SVirtualAssetsStatistics.cpp
+++ b/Engine/Source/Editor/DerivedDataEditor/Private/SVirtualAssetsStatistics.cpp
@@ -24,8 +24,6 @@
 
 extern FString SingleDecimalFormat(double Value);
 
-using namespace UE::Virtualization;
-
 // This namespace contains code for adding the re-hydration options to the editor context menu.
 // For now it is easier to just put all virtualization editor code in this .cpp until we make
 // a virtualization specific editor module like 'DerivedDataEditor'
@@ -147,6 +145,8 @@ void AddContextMenuEntry()
 
 SVirtualAssetsStatisticsDialog::SVirtualAssetsStatisticsDialog()
 {
+	using namespace UE::Virtualization;
+
 	// Register our VA notification delegate with the event
 	IVirtualizationSystem& System = IVirtualizationSystem::Get();
 	System.GetNotificationEvent().AddRaw(this, &SVirtualAssetsStatisticsDialog::OnNotificationEvent);
@@ -156,13 +156,17 @@ SVirtualAssetsStatisticsDialog::SVirtualAssetsStatisticsDialog()
 
 SVirtualAssetsStatisticsDialog::~SVirtualAssetsStatisticsDialog()
 {
+	using namespace UE::Virtualization;
+
 	// Unregister our VA notification delegate with the event
 	IVirtualizationSystem& System = IVirtualizationSystem::Get();
 	System.GetNotificationEvent().RemoveAll(this);
 }
 
-void SVirtualAssetsStatisticsDialog::OnNotificationEvent(IVirtualizationSystem::ENotification Notification, const FIoHash& PayloadId)
+void SVirtualAssetsStatisticsDialog::OnNotificationEvent(UE::Virtualization::IVirtualizationSystem::ENotification Notification, const FIoHash& PayloadId)
 {
+	using namespace UE::Virtualization;
+
 	FScopeLock SocpeLock(&NotificationCS);
 	
 	switch (Notification)
@@ -309,6 +313,8 @@ void SVirtualAssetsStatisticsDialog::OnWarningReasonIgnore()
 
 TSharedRef<SWidget> SVirtualAssetsStatisticsDialog::GetGridPanel()
 {
+	using namespace UE::Virtualization;
+
 	IVirtualizationSystem& System = IVirtualizationSystem::Get();
 
 	TSharedRef<SGridPanel> Panel = SNew(SGridPanel);
