Add status icon to indicate DDC cache store status with tooltip text for details.  Currently used to show when a cache store has been deactivated for performance.

diff --git a/Engine/Source/Developer/DerivedDataCache/Private/FileSystemCacheStore.cpp b/Engine/Source/Developer/DerivedDataCache/Private/FileSystemCacheStore.cpp
--- a/Engine/Source/Developer/DerivedDataCache/Private/FileSystemCacheStore.cpp
+++ b/Engine/Source/Developer/DerivedDataCache/Private/FileSystemCacheStore.cpp
@@ -1326,7 +1326,15 @@ bool FFileSystemCacheStore::RunSpeedTest(
 
 void FFileSystemCacheStore::LegacyStats(FDerivedDataCacheStatsNode& OutNode)
 {
-	OutNode = {TEXT("File System"), *CachePath, SpeedClass == EBackendSpeedClass::Local};
+	EDerivedDataCacheStatus CacheStatus = EDerivedDataCacheStatus::None;
+	FText CacheStatusText;
+	if (IsDeactivatedForPerformance())
+	{
+		CacheStatus = EDerivedDataCacheStatus::Deactivation;
+		CacheStatusText = NSLOCTEXT("DerivedDataCache", "DeactivatedForPerformance", "Deactivated for performance");
+	}
+	const FString& CacheStatusTextString = CacheStatusText.ToString();
+	OutNode = {TEXT("File System"), *CachePath, SpeedClass == EBackendSpeedClass::Local, CacheStatus, CacheStatusTextString.IsEmpty() ? nullptr : *CacheStatusTextString};
 	OutNode.UsageStats.Add(TEXT(""), UsageStats);
 	OutNode.SpeedStats = SpeedStats;
 }
diff --git a/Engine/Source/Developer/DerivedDataCache/Public/DerivedDataCacheUsageStats.h b/Engine/Source/Developer/DerivedDataCache/Public/DerivedDataCacheUsageStats.h
--- a/Engine/Source/Developer/DerivedDataCache/Public/DerivedDataCacheUsageStats.h
+++ b/Engine/Source/Developer/DerivedDataCache/Public/DerivedDataCacheUsageStats.h
@@ -111,6 +111,15 @@ struct FDerivedDataCacheSpeedStats
 	double		LatencyMS = 0.0;
 };
 
+enum class EDerivedDataCacheStatus
+{
+	None = 0,
+	Information,
+	Warning,
+	Error,
+	Deactivation
+};
+
 /**
  *  Hierarchical usage stats for the DDC nodes.
  */
@@ -119,9 +128,11 @@ class FDerivedDataCacheStatsNode : public TSharedFromThis<FDerivedDataCacheStats
 public:
 	FDerivedDataCacheStatsNode() = default;
 
-	FDerivedDataCacheStatsNode(const FString& InCacheType, const FString& InCacheName, bool bInIsLocal)
+	FDerivedDataCacheStatsNode(const FString& InCacheType, const FString& InCacheName, bool bInIsLocal, EDerivedDataCacheStatus InCacheStatus = EDerivedDataCacheStatus::None, const TCHAR* InCacheStatusText = nullptr)
 		: CacheType(InCacheType)
 		, CacheName(InCacheName)
+		, CacheStatusText(InCacheStatusText)
+		, CacheStatus(InCacheStatus)
 		, bIsLocal(bInIsLocal)
 	{
 	}
@@ -130,6 +141,10 @@ public:
 
 	const FString& GetCacheName() const { return CacheName; }
 
+	const EDerivedDataCacheStatus GetCacheStatus() const { return CacheStatus; }
+
+	const FString& GetCacheStatusText() const { return CacheStatusText; }
+
 	bool IsLocal() const { return bIsLocal; }
 
 	TMap<FString, FDerivedDataCacheUsageStats> ToLegacyUsageMap() const
@@ -182,6 +197,8 @@ public:
 protected:
 	FString CacheType;
 	FString CacheName;
+	FString CacheStatusText;
+	EDerivedDataCacheStatus CacheStatus;
 	bool bIsLocal;
 };
 
diff --git a/Engine/Source/Editor/DerivedDataEditor/Private/SDerivedDataDialogs.cpp b/Engine/Source/Editor/DerivedDataEditor/Private/SDerivedDataDialogs.cpp
--- a/Engine/Source/Editor/DerivedDataEditor/Private/SDerivedDataDialogs.cpp
+++ b/Engine/Source/Editor/DerivedDataEditor/Private/SDerivedDataDialogs.cpp
@@ -538,16 +538,20 @@ TSharedRef<SWidget> SDerivedDataCacheStatisticsDialog::GetGridPanel()
 	const FMargin DefaultMargin(0.0f, RowMargin, ColumnMargin, RowMargin);
 	
 	Panel->AddSlot(0, Row)
+	[
+		SNew(STextBlock)
+	];
+
+	Panel->AddSlot(1, Row)
 	[
 		SNew(STextBlock)
 		.Margin(TitleMarginFirstColumn)
 		.Font(FCoreStyle::GetDefaultFontStyle("Bold", 10))
-		.Text(LOCTEXT("Cache", "Cache"))
 		.ColorAndOpacity(TitleColor)
 		.Text(LOCTEXT("CacheType", "Cache Type"))
 	];
 
-	Panel->AddSlot(1, Row)
+	Panel->AddSlot(2, Row)
 	[
 		SNew(STextBlock)
 		.Margin(TitleMargin)
@@ -556,7 +560,7 @@ TSharedRef<SWidget> SDerivedDataCacheStatisticsDialog::GetGridPanel()
 		.Text(LOCTEXT("Location", "Location"))
 	];
 
-	Panel->AddSlot(2, Row)
+	Panel->AddSlot(3, Row)
 	.HAlign(HAlign_Right)
 	[
 		SNew(STextBlock)
@@ -566,7 +570,7 @@ TSharedRef<SWidget> SDerivedDataCacheStatisticsDialog::GetGridPanel()
 		.Text(LOCTEXT("HitPercentage", "Hit%"))
 	];
 
-	Panel->AddSlot(3, Row)
+	Panel->AddSlot(4, Row)
 	.HAlign(HAlign_Right)
 	[
 		SNew(STextBlock)
@@ -578,7 +582,7 @@ TSharedRef<SWidget> SDerivedDataCacheStatisticsDialog::GetGridPanel()
 		.WrapTextAt(66.0f)
 	];
 
-	Panel->AddSlot(4, Row)
+	Panel->AddSlot(5, Row)
 	.HAlign(HAlign_Right)
 	[
 		SNew(STextBlock)
@@ -590,7 +594,7 @@ TSharedRef<SWidget> SDerivedDataCacheStatisticsDialog::GetGridPanel()
 		.WrapTextAt(66.0f)
 	];
 
-	Panel->AddSlot(5, Row)
+	Panel->AddSlot(6, Row)
 	.HAlign(HAlign_Right)
 	[
 		SNew(STextBlock)
@@ -602,7 +606,7 @@ TSharedRef<SWidget> SDerivedDataCacheStatisticsDialog::GetGridPanel()
 		.WrapTextAt(66.0f)
 	];
 
-	Panel->AddSlot(6, Row)
+	Panel->AddSlot(7, Row)
 	.HAlign(HAlign_Right)
 	[
 		SNew(STextBlock)
@@ -614,7 +618,7 @@ TSharedRef<SWidget> SDerivedDataCacheStatisticsDialog::GetGridPanel()
 		.WrapTextAt(66.0f)
 	];
 
-	Panel->AddSlot(7, Row)
+	Panel->AddSlot(8, Row)
 	.HAlign(HAlign_Right)
 	[
 		SNew(STextBlock)
@@ -626,7 +630,7 @@ TSharedRef<SWidget> SDerivedDataCacheStatisticsDialog::GetGridPanel()
 		.WrapTextAt(66.0f)
 	];
 
-	Panel->AddSlot(8, Row)
+	Panel->AddSlot(9, Row)
 	[
 		SNew(STextBlock)
 		.Margin(TitleMargin)
@@ -668,21 +672,70 @@ TSharedRef<SWidget> SDerivedDataCacheStatisticsDialog::GetGridPanel()
 		SumTotalGetMB += TotalGetMB;
 		SumTotalPutMB += TotalPutMB;
 
-		Panel->AddSlot(0, Row)
+		TSharedPtr<SImage> StatusIcon;
+		switch (Node->GetCacheStatus())
+		{
+			case EDerivedDataCacheStatus::Information:
+				StatusIcon = SNew(SImage)
+					.Image(FAppStyle::GetBrush("Icons.Help"))
+					.ToolTipText(FText::FromString(Node->GetCacheStatusText()));
+			break;
+			case EDerivedDataCacheStatus::Warning:
+				StatusIcon = SNew(SImage)
+					.Image(FAppStyle::GetBrush("Icons.Warning"))
+					.ToolTipText(FText::FromString(Node->GetCacheStatusText()));
+			break;
+			case EDerivedDataCacheStatus::Error:
+				StatusIcon = SNew(SImage)
+					.Image(FAppStyle::GetBrush("Icons.Error"))
+					.ToolTipText(FText::FromString(Node->GetCacheStatusText()));
+			break;
+			case EDerivedDataCacheStatus::Deactivation:
+				StatusIcon = SNew(SImage)
+					.Image(FAppStyle::GetBrush("Icons.X"))
+					.ToolTipText(FText::FromString(Node->GetCacheStatusText()));
+			break;
+
+		}
+
+		if (StatusIcon.IsValid())
+		{
+			Panel->AddSlot(0, Row)
+			[
+				SNew(SHorizontalBox)
+					.ToolTipText(FText::FromString(Node->GetCacheStatusText()))
+					// Status icon
+					+SHorizontalBox::Slot()
+					.AutoWidth()
+					.VAlign(VAlign_Center)
+					[
+						StatusIcon.ToSharedRef()
+					]
+			];
+		}
+		else
+		{
+			Panel->AddSlot(0, Row)
+			[
+				SNew(SHorizontalBox)
+			];
+		}
+
+		Panel->AddSlot(1, Row)
 		[
 			SNew(STextBlock)
 			.Margin(DefaultMarginFirstColumn)
 			.Text(FText::FromString(Node->GetCacheType()))
 		];
 
-		Panel->AddSlot(1, Row)
+		Panel->AddSlot(2, Row)
 		[
 			SNew(STextBlock)
 			.Margin(DefaultMargin)
 			.Text(Node->IsLocal() ? LOCTEXT("Local", "Local") : LOCTEXT("Remote", "Remote"))
 		];
 
-		Panel->AddSlot(2, Row)
+		Panel->AddSlot(3, Row)
 		.HAlign(HAlign_Right)
 		[
 			SNew(STextBlock)
@@ -690,7 +743,7 @@ TSharedRef<SWidget> SDerivedDataCacheStatisticsDialog::GetGridPanel()
 			.Text(FText::FromString(SingleDecimalFormat(HitRate)))
 		];
 
-		Panel->AddSlot(3, Row)
+		Panel->AddSlot(4, Row)
 		.HAlign(HAlign_Right)
 		[
 			SNew(STextBlock)
@@ -698,7 +751,7 @@ TSharedRef<SWidget> SDerivedDataCacheStatisticsDialog::GetGridPanel()
 			.Text(FText::FromString(SingleDecimalFormat(TotalGetMB)))
 		];
 
-		Panel->AddSlot(4, Row)
+		Panel->AddSlot(5, Row)
 		.HAlign(HAlign_Right)
 		[
 			SNew(STextBlock)
@@ -708,7 +761,7 @@ TSharedRef<SWidget> SDerivedDataCacheStatisticsDialog::GetGridPanel()
 
 		if (Node->SpeedStats.LatencyMS)
 		{
-			Panel->AddSlot(5, Row)
+			Panel->AddSlot(6, Row)
 			.HAlign(HAlign_Right)
 			[
 				SNew(STextBlock)
@@ -719,7 +772,7 @@ TSharedRef<SWidget> SDerivedDataCacheStatisticsDialog::GetGridPanel()
 
 		if (Node->SpeedStats.ReadSpeedMBs)
 		{
-			Panel->AddSlot(6, Row)
+			Panel->AddSlot(7, Row)
 			.HAlign(HAlign_Right)
 			[
 				SNew(STextBlock)
@@ -730,7 +783,7 @@ TSharedRef<SWidget> SDerivedDataCacheStatisticsDialog::GetGridPanel()
 
 		if (Node->SpeedStats.WriteSpeedMBs)
 		{
-			Panel->AddSlot(7, Row)
+			Panel->AddSlot(8, Row)
 			.HAlign(HAlign_Right)
 			[
 				SNew(STextBlock)
@@ -739,7 +792,7 @@ TSharedRef<SWidget> SDerivedDataCacheStatisticsDialog::GetGridPanel()
 			];
 		}
 
-		Panel->AddSlot(8, Row)
+		Panel->AddSlot(9, Row)
 		[
 			SNew(STextBlock)
 			.Margin(DefaultMargin)
@@ -749,7 +802,7 @@ TSharedRef<SWidget> SDerivedDataCacheStatisticsDialog::GetGridPanel()
 		Row++;
 	}
 
-	Panel->AddSlot(0, Row)
+	Panel->AddSlot(1, Row)
 	[
 		SNew(STextBlock)
 		.Margin(TitleMarginFirstColumn)
@@ -758,7 +811,7 @@ TSharedRef<SWidget> SDerivedDataCacheStatisticsDialog::GetGridPanel()
 		.Text(LOCTEXT("Total", "Total"))
 	];
 
-	Panel->AddSlot(3, Row)
+	Panel->AddSlot(4, Row)
 	.HAlign(HAlign_Right)
 	[
 		SNew(STextBlock)
@@ -768,7 +821,7 @@ TSharedRef<SWidget> SDerivedDataCacheStatisticsDialog::GetGridPanel()
 		.Text(FText::FromString(SingleDecimalFormat(SumTotalGetMB)))
 	];
 
-	Panel->AddSlot(4, Row)
+	Panel->AddSlot(5, Row)
 	.HAlign(HAlign_Right)
 	[
 		SNew(STextBlock)
