Improve the message shown by the VA statistics panel if the system has not yet initialized because it has not yet been used.

#rnx

diff --git a/Engine/Source/Editor/DerivedDataEditor/Private/SVirtualAssetsStatistics.cpp b/Engine/Source/Editor/DerivedDataEditor/Private/SVirtualAssetsStatistics.cpp
--- a/Engine/Source/Editor/DerivedDataEditor/Private/SVirtualAssetsStatistics.cpp
+++ b/Engine/Source/Editor/DerivedDataEditor/Private/SVirtualAssetsStatistics.cpp
@@ -336,7 +336,7 @@ TSharedRef<SWidget> SVirtualAssetsStatisticsDialog::GetGridPanel()
 				.ColorAndOpacity(TitleColor)
 				.Font(TitleFont)
 				.Justification(ETextJustify::Center)
-				.Text(LOCTEXT("LazyInit", "Virtual Assets are lazy initialized for this project, cannot display stats"))
+				.Text(LOCTEXT("LazyInit", "System is uninitialized until first use"))
 			];
 
 		return Panel;
