Update references to "Horde Storage" to "Unreal Cloud DDC" or appropriate variations.
Specifically:
HordeStorageServers-> StorageServers
HordeStorage  -> UnrealCloudDDC
Horde Storage -> Unreal Cloud DDC

diff --git a/Engine/Config/BaseGame.ini b/Engine/Config/BaseGame.ini
--- a/Engine/Config/BaseGame.ini
+++ b/Engine/Config/BaseGame.ini
@@ -149,6 +149,7 @@ WriteBackMetadataToAssetRegistry=Disabled
 +IniKeyDenylist=IniSectionDenylist
 
 +IniSectionDenylist=HordeStorageServers
++IniSectionDenylist=StorageServers
 
 [/Script/Engine.HUD]
 DebugDisplay=AI
diff --git a/Engine/Source/Developer/DerivedDataCache/Private/DerivedDataCache.cpp b/Engine/Source/Developer/DerivedDataCache/Private/DerivedDataCache.cpp
--- a/Engine/Source/Developer/DerivedDataCache/Private/DerivedDataCache.cpp
+++ b/Engine/Source/Developer/DerivedDataCache/Private/DerivedDataCache.cpp
@@ -89,7 +89,7 @@ namespace UE::DerivedData::CookStats
 		{
 			const TSharedRef<const FDerivedDataCacheStatsNode>* LocalNode = Nodes.FindByPredicate([](TSharedRef<const FDerivedDataCacheStatsNode> Node) { return Node->GetCacheType() == TEXT("File System") && Node->IsLocal(); });
 			const TSharedRef<const FDerivedDataCacheStatsNode>* SharedNode = Nodes.FindByPredicate([](TSharedRef<const FDerivedDataCacheStatsNode> Node) { return Node->GetCacheType() == TEXT("File System") && !Node->IsLocal(); });
-			const TSharedRef<const FDerivedDataCacheStatsNode>* CloudNode = Nodes.FindByPredicate([](TSharedRef<const FDerivedDataCacheStatsNode> Node) { return Node->GetCacheType() == TEXT("Horde Storage"); });
+			const TSharedRef<const FDerivedDataCacheStatsNode>* CloudNode = Nodes.FindByPredicate([](TSharedRef<const FDerivedDataCacheStatsNode> Node) { return Node->GetCacheType() == TEXT("Unreal Cloud DDC"); });
 			const TSharedRef<const FDerivedDataCacheStatsNode>* ZenLocalNode = Nodes.FindByPredicate([](TSharedRef<const FDerivedDataCacheStatsNode> Node) { return Node->GetCacheType() == TEXT("Zen") && Node->IsLocal(); });
 			const TSharedRef<const FDerivedDataCacheStatsNode>* ZenRemoteNode = Nodes.FindByPredicate([](TSharedRef<const FDerivedDataCacheStatsNode> Node) { return (Node->GetCacheType() == TEXT("Zen") || Node->GetCacheType() == TEXT("Horde")) && !Node->IsLocal(); });
 
diff --git a/Engine/Source/Developer/DerivedDataCache/Private/HttpCacheStore.cpp b/Engine/Source/Developer/DerivedDataCache/Private/HttpCacheStore.cpp
--- a/Engine/Source/Developer/DerivedDataCache/Private/HttpCacheStore.cpp
+++ b/Engine/Source/Developer/DerivedDataCache/Private/HttpCacheStore.cpp
@@ -2266,7 +2266,7 @@ void FHttpCacheStore::RefCachedDataProbablyExistsBatchAsync(
 
 void FHttpCacheStore::LegacyStats(FDerivedDataCacheStatsNode& OutNode)
 {
-	OutNode = {TEXT("Horde Storage"), FString::Printf(TEXT("%s (%s)"), *Domain, *Namespace), /*bIsLocal*/ false};
+	OutNode = {TEXT("Unreal Cloud DDC"), FString::Printf(TEXT("%s (%s)"), *Domain, *Namespace), /*bIsLocal*/ false};
 	OutNode.UsageStats.Add(TEXT(""), UsageStats);
 }
 
@@ -2611,11 +2611,16 @@ void FHttpCacheStoreParams::Parse(const TCHAR* NodeName, const TCHAR* Config)
 	if (FParse::Value(Config, TEXT("ServerID="), ServerId))
 	{
 		FString ServerEntry;
-		const TCHAR* ServerSection = TEXT("HordeStorageServers");
+		const TCHAR* ServerSection = TEXT("StorageServers");
+		const TCHAR* FallbackServerSection = TEXT("HordeStorageServers");
 		if (GConfig->GetString(ServerSection, *ServerId, ServerEntry, GEngineIni))
 		{
 			Parse(NodeName, *ServerEntry);
 		}
+		else if (GConfig->GetString(FallbackServerSection, *ServerId, ServerEntry, GEngineIni))
+		{
+			Parse(NodeName, *ServerEntry);
+		}
 		else
 		{
 			UE_LOG(LogDerivedDataCache, Warning, TEXT("%s: Using ServerID=%s which was not found in [%s]"), NodeName, *ServerId, ServerSection);
diff --git a/Engine/Source/Editor/DerivedDataEditor/Private/DerivedDataInformation.cpp b/Engine/Source/Editor/DerivedDataEditor/Private/DerivedDataInformation.cpp
--- a/Engine/Source/Editor/DerivedDataEditor/Private/DerivedDataInformation.cpp
+++ b/Engine/Source/Editor/DerivedDataEditor/Private/DerivedDataInformation.cpp
@@ -128,11 +128,11 @@ bool FDerivedDataInformation::GetHasZenCache()
 	return false;
 }
 
-bool FDerivedDataInformation::GetHasHordeStorageCache()
+bool FDerivedDataInformation::GetHasUnrealCloudCache()
 {
 	for (const TSharedRef<const FDerivedDataCacheStatsNode>& Usage : GetCacheUsageStats())
 		{
-		if (Usage->GetCacheType().Equals(TEXT("Horde Storage")))
+		if (Usage->GetCacheType().Equals(TEXT("Unreal Cloud DDC")))
 	{
 			return true;
 	}
@@ -174,10 +174,10 @@ void FDerivedDataInformation::UpdateRemoteCacheState()
 
 		if (EditorSettings && EditorSettings->bEnableDDCNotifications)
 		{
-			if (DDCProjectSettings->RecommendEveryoneUseHordeStorage && EditorSettings->bNotifyUseHordeStorage && !GetHasZenCache() && !GetHasHordeStorageCache() && (FCString::Stricmp(GetDerivedDataCache()->GetGraphName(), TEXT("NoJupiter")) != 0))
+			if (DDCProjectSettings->RecommendEveryoneUseUnrealCloudDDC && EditorSettings->bNotifyUseUnrealCloudDDC && !GetHasZenCache() && !GetHasUnrealCloudCache() && (FCString::Stricmp(GetDerivedDataCache()->GetGraphName(), TEXT("NoJupiter")) != 0))
 			{
 				RemoteCacheState = ERemoteCacheState::Warning;
-				RemoteCacheWarningMessage = FText(LOCTEXT("HordeStorageWarning", "It is recommended that you use a DDC graph that supports Horde Storage. Please check any -ddc commandline overrides."));
+				RemoteCacheWarningMessage = FText(LOCTEXT("UnrealCloudDDCWarning", "It is recommended that you use a DDC graph that supports Unreal Cloud DDC. Please check any -ddc commandline overrides."));
 			}
 			else if (DDCProjectSettings->RecommendEveryoneSetupAGlobalLocalDDCPath && EditorSettings->bNotifySetupDDCPath && EditorSettings->GlobalLocalDDCPath.Path.IsEmpty())
 			{
diff --git a/Engine/Source/Editor/DerivedDataEditor/Private/DerivedDataInformation.h b/Engine/Source/Editor/DerivedDataEditor/Private/DerivedDataInformation.h
--- a/Engine/Source/Editor/DerivedDataEditor/Private/DerivedDataInformation.h
+++ b/Engine/Source/Editor/DerivedDataEditor/Private/DerivedDataInformation.h
@@ -21,7 +21,7 @@ public:
 	static double				GetCacheActivitySizeBytes(bool bGet, bool bLocal);
 	static bool					GetHasRemoteCache();
 	static bool					GetHasZenCache();
-	static bool					GetHasHordeStorageCache();
+	static bool					GetHasUnrealCloudCache();
 	static ERemoteCacheState	GetRemoteCacheState() { return RemoteCacheState; }
 	static FText				GetRemoteCacheStateAsText();
 	static FText				GetRemoteCacheWarningMessage() { return RemoteCacheWarningMessage; }
diff --git a/Engine/Source/Editor/UnrealEd/Public/Settings/EditorProjectSettings.h b/Engine/Source/Editor/UnrealEd/Public/Settings/EditorProjectSettings.h
--- a/Engine/Source/Editor/UnrealEd/Public/Settings/EditorProjectSettings.h
+++ b/Engine/Source/Editor/UnrealEd/Public/Settings/EditorProjectSettings.h
@@ -268,5 +268,5 @@ public:
 	 *
 	 */
 	UPROPERTY(EditAnywhere, config, Category = Warnings)
-	bool RecommendEveryoneUseHordeStorage = false;
+	bool RecommendEveryoneUseUnrealCloudDDC = false;
 };
diff --git a/Engine/Source/Editor/UnrealEd/Public/Settings/EditorSettings.h b/Engine/Source/Editor/UnrealEd/Public/Settings/EditorSettings.h
--- a/Engine/Source/Editor/UnrealEd/Public/Settings/EditorSettings.h
+++ b/Engine/Source/Editor/UnrealEd/Public/Settings/EditorSettings.h
@@ -92,9 +92,9 @@ class UNREALED_API UEditorSettings : public UObject
 	UPROPERTY(EditAnywhere, config, Category = "Derived Data Cache Notifications", meta = (DisplayName = "Enable Notifcations", ConfigRestartRequired = false))
 	bool bEnableDDCNotifications = true;
 
-	/** Whether to enable the Horde Storage notification */
-	UPROPERTY(EditAnywhere, config, Category = "Derived Data Cache Notifications", meta = (DisplayName = "Notify Use Horde Storage", ConfigRestartRequired = false, EditCondition = "bEnableDDCNotifications"))
-	bool bNotifyUseHordeStorage = true;
+	/** Whether to enable the Unreal Cloud DDC notification */
+	UPROPERTY(EditAnywhere, config, Category = "Derived Data Cache Notifications", meta = (DisplayName = "Notify Use Unreal Cloud DDC", ConfigRestartRequired = false, EditCondition = "bEnableDDCNotifications"))
+	bool bNotifyUseUnrealCloudDDC = true;
 
 	/** Whether to enable the DDC path notification */
 	UPROPERTY(EditAnywhere, config, Category = "Derived Data Cache Notifications", meta = (DisplayName = "Notify Setup DDC Path", ConfigRestartRequired = false, EditCondition = "bEnableDDCNotifications"))
