Activating truncation warnings for
Analytics, AnalyticsET, Cbor, Serialization, XMLParser, DerivedDataEditor, DirectoryWatcher, LogVisualizer, InputCore

fixing a few addtional truncation warnings in SDerivedDataDialogs

diff --git a/Engine/Source/Developer/DirectoryWatcher/DirectoryWatcher.Build.cs b/Engine/Source/Developer/DirectoryWatcher/DirectoryWatcher.Build.cs
--- a/Engine/Source/Developer/DirectoryWatcher/DirectoryWatcher.Build.cs
+++ b/Engine/Source/Developer/DirectoryWatcher/DirectoryWatcher.Build.cs
@@ -7,5 +7,6 @@ public class DirectoryWatcher : ModuleRules
 	public DirectoryWatcher(ReadOnlyTargetRules Target) : base(Target)
 	{
 		PrivateDependencyModuleNames.Add("Core");
+		UnsafeTypeCastWarningLevel = WarningLevel.Error;
 	}
 }
diff --git a/Engine/Source/Editor/DerivedDataEditor/DerivedDataEditor.Build.cs b/Engine/Source/Editor/DerivedDataEditor/DerivedDataEditor.Build.cs
--- a/Engine/Source/Editor/DerivedDataEditor/DerivedDataEditor.Build.cs
+++ b/Engine/Source/Editor/DerivedDataEditor/DerivedDataEditor.Build.cs
@@ -39,5 +39,6 @@ public class DerivedDataEditor : ModuleRules
 			new string[] {
 				"MainFrame",
 			});
+		UnsafeTypeCastWarningLevel = WarningLevel.Error;
 	}
 }
diff --git a/Engine/Source/Editor/DerivedDataEditor/Private/SDerivedDataDialogs.cpp b/Engine/Source/Editor/DerivedDataEditor/Private/SDerivedDataDialogs.cpp
--- a/Engine/Source/Editor/DerivedDataEditor/Private/SDerivedDataDialogs.cpp
+++ b/Engine/Source/Editor/DerivedDataEditor/Private/SDerivedDataDialogs.cpp
@@ -83,7 +83,7 @@ TSharedRef<SWidget> SDerivedDataRemoteStoreDialog::GetGridPanel()
 	}
 
 	const int64 TotalCount = DDCResourceStatsTotal.LoadCount + DDCResourceStatsTotal.BuildCount;
-	const double Efficiency = TotalCount > 0 ? (double)DDCResourceStatsTotal.LoadCount / TotalCount : 0.0;
+	const double Efficiency = TotalCount > 0 ? static_cast<double>(DDCResourceStatsTotal.LoadCount) / static_cast<double>(TotalCount) : 0.0;
 
 	const double DownloadedBytesMB = FUnitConversion::Convert(FDerivedDataInformation::GetCacheActivitySizeBytes(true, false), EUnit::Bytes, EUnit::Megabytes);
 	const double UploadedBytesMB = FUnitConversion::Convert(FDerivedDataInformation::GetCacheActivitySizeBytes(false, false), EUnit::Bytes, EUnit::Megabytes);
@@ -664,10 +664,10 @@ TSharedRef<SWidget> SDerivedDataCacheStatisticsDialog::GetGridPanel()
 		const int64 TotalGetHits = Stats.GetStats.GetAccumulatedValueAnyThread(FCookStats::CallStats::EHitOrMiss::Hit, FCookStats::CallStats::EStatType::Counter);
 		const int64 TotalGetMisses = Stats.GetStats.GetAccumulatedValueAnyThread(FCookStats::CallStats::EHitOrMiss::Miss, FCookStats::CallStats::EStatType::Counter);
 		const int64 TotalRequests = TotalGetHits + TotalGetMisses;
-		const double HitRate = TotalRequests > 0 ? 100.0 * TotalGetHits / TotalRequests : 0.0;
+		const double HitRate = TotalRequests > 0 ? 100.0 * static_cast<double>(TotalGetHits) / static_cast<double>(TotalRequests) : 0.0;
 
-		const double TotalGetMB = FUnitConversion::Convert((double)TotalGetBytes, EUnit::Bytes, EUnit::Megabytes);
-		const double TotalPutMB = FUnitConversion::Convert((double)TotalPutBytes, EUnit::Bytes, EUnit::Megabytes);
+		const double TotalGetMB = FUnitConversion::Convert(static_cast<double>(TotalGetBytes), EUnit::Bytes, EUnit::Megabytes);
+		const double TotalPutMB = FUnitConversion::Convert(static_cast<double>(TotalPutBytes), EUnit::Bytes, EUnit::Megabytes);
 
 		SumTotalGetMB += TotalGetMB;
 		SumTotalPutMB += TotalPutMB;
@@ -759,7 +759,7 @@ TSharedRef<SWidget> SDerivedDataCacheStatisticsDialog::GetGridPanel()
 			.Text(FText::FromString(SingleDecimalFormat(TotalPutMB)))
 		];
 
-		if (Node->SpeedStats.LatencyMS)
+		if (Node->SpeedStats.LatencyMS != 0.0)
 		{
 			Panel->AddSlot(6, Row)
 			.HAlign(HAlign_Right)
@@ -770,7 +770,7 @@ TSharedRef<SWidget> SDerivedDataCacheStatisticsDialog::GetGridPanel()
 			];
 		}
 
-		if (Node->SpeedStats.ReadSpeedMBs)
+		if (Node->SpeedStats.ReadSpeedMBs != 0.0)
 		{
 			Panel->AddSlot(7, Row)
 			.HAlign(HAlign_Right)
@@ -781,7 +781,7 @@ TSharedRef<SWidget> SDerivedDataCacheStatisticsDialog::GetGridPanel()
 			];
 		}
 
-		if (Node->SpeedStats.WriteSpeedMBs)
+		if (Node->SpeedStats.WriteSpeedMBs != 0.0)
 		{
 			Panel->AddSlot(8, Row)
 			.HAlign(HAlign_Right)
diff --git a/Engine/Source/Runtime/Analytics/Analytics/Analytics.Build.cs b/Engine/Source/Runtime/Analytics/Analytics/Analytics.Build.cs
--- a/Engine/Source/Runtime/Analytics/Analytics/Analytics.Build.cs
+++ b/Engine/Source/Runtime/Analytics/Analytics/Analytics.Build.cs
@@ -13,6 +13,7 @@ namespace UnrealBuildTool.Rules
                     // ... add other public dependencies that you statically link with here ...
 				}
 				);
+			UnsafeTypeCastWarningLevel = WarningLevel.Error;
 		}
 	}
 }
diff --git a/Engine/Source/Runtime/Analytics/AnalyticsET/AnalyticsET.Build.cs b/Engine/Source/Runtime/Analytics/AnalyticsET/AnalyticsET.Build.cs
--- a/Engine/Source/Runtime/Analytics/AnalyticsET/AnalyticsET.Build.cs
+++ b/Engine/Source/Runtime/Analytics/AnalyticsET/AnalyticsET.Build.cs
@@ -24,6 +24,7 @@ namespace UnrealBuildTool.Rules
 				);
 
 			PublicIncludePathModuleNames.Add("Analytics");
+			UnsafeTypeCastWarningLevel = WarningLevel.Error;
 		}
 	}
 }
diff --git a/Engine/Source/Runtime/Cbor/Cbor.Build.cs b/Engine/Source/Runtime/Cbor/Cbor.Build.cs
--- a/Engine/Source/Runtime/Cbor/Cbor.Build.cs
+++ b/Engine/Source/Runtime/Cbor/Cbor.Build.cs
@@ -12,6 +12,7 @@ namespace UnrealBuildTool.Rules
 					"Core",
 				}
 			);
-        }
+			UnsafeTypeCastWarningLevel = WarningLevel.Error;
+		}
     }
 }
diff --git a/Engine/Source/Runtime/InputCore/InputCore.Build.cs b/Engine/Source/Runtime/InputCore/InputCore.Build.cs
--- a/Engine/Source/Runtime/InputCore/InputCore.Build.cs
+++ b/Engine/Source/Runtime/InputCore/InputCore.Build.cs
@@ -16,5 +16,6 @@ public class InputCore : ModuleRules
 		{
 			PrivateIncludePathModuleNames.Add("SDL2");
 		}
+		UnsafeTypeCastWarningLevel = WarningLevel.Error;
 	}
 }
diff --git a/Engine/Source/Runtime/Serialization/Serialization.Build.cs b/Engine/Source/Runtime/Serialization/Serialization.Build.cs
--- a/Engine/Source/Runtime/Serialization/Serialization.Build.cs
+++ b/Engine/Source/Runtime/Serialization/Serialization.Build.cs
@@ -17,6 +17,7 @@ namespace UnrealBuildTool.Rules
 				new string[] {
 					"CoreUObject",
 				});
+			UnsafeTypeCastWarningLevel = WarningLevel.Error;
 		}
 	}
 }
diff --git a/Engine/Source/Runtime/XmlParser/XmlParser.Build.cs b/Engine/Source/Runtime/XmlParser/XmlParser.Build.cs
--- a/Engine/Source/Runtime/XmlParser/XmlParser.Build.cs
+++ b/Engine/Source/Runtime/XmlParser/XmlParser.Build.cs
@@ -10,7 +10,7 @@ public class XmlParser : ModuleRules
 			new string[] 
 			{ 
 				"Core",
-			}
-			);
+			});
+		UnsafeTypeCastWarningLevel = WarningLevel.Error;
 	}
 }
