Tweak to DDC Cache notification UX
Improved notification messages
Deprecated unsupported shared cache notification warning

diff --git a/Engine/Source/Editor/DerivedDataEditor/Private/DerivedDataEditorModule.cpp b/Engine/Source/Editor/DerivedDataEditor/Private/DerivedDataEditorModule.cpp
--- a/Engine/Source/Editor/DerivedDataEditor/Private/DerivedDataEditorModule.cpp
+++ b/Engine/Source/Editor/DerivedDataEditor/Private/DerivedDataEditorModule.cpp
@@ -6,7 +6,6 @@
 #include "Experimental/ZenServerInterface.h"
 #include "Framework/Application/SlateApplication.h"
 #include "Modules/ModuleManager.h"
-#include "SDerivedDataCacheSettings.h"
 #include "SDerivedDataDialogs.h"
 #include "SDerivedDataStatusBar.h"
 #include "Styling/AppStyle.h"
@@ -121,44 +120,4 @@ void FDerivedDataEditorModule::ShowCacheStatisticsTab()
 	FGlobalTabmanager::Get()->TryInvokeTab(FTabId(DerivedDataCacheStatisticsTabName));
 }
 
-void FDerivedDataEditorModule::ShowSettingsDialog()
-{
-	if (SettingsWindow.IsValid())
-	{
-		SettingsWindow->BringToFront();
-	}
-	else
-	{
-		// Create the window
-		SettingsWindow = SNew(SWindow)
-			.Title(LOCTEXT("DerrivedDataCacheSettingsWindowTitle", "Cache Settings"))
-			.HasCloseButton(true)
-			.SupportsMaximize(false)
-			.SupportsMinimize(true)
-			.SizingRule(ESizingRule::Autosized);
-
-		// Set the closed callback
-		SettingsWindow->SetOnWindowClosed(FOnWindowClosed::CreateRaw(this,&FDerivedDataEditorModule::OnSettingsDialogClosed));
-
-		// Setup the content for the created window.
-		SettingsWindow->SetContent(SAssignNew(SettingsDialog, SDerivedDataCacheSettingsDialog));
-
-		TSharedPtr<SWindow> RootWindow = FGlobalTabmanager::Get()->GetRootWindow();
-		if (RootWindow.IsValid())
-		{
-			FSlateApplication::Get().AddWindowAsNativeChild(SettingsWindow.ToSharedRef(), RootWindow.ToSharedRef());
-		}
-		else
-		{
-			FSlateApplication::Get().AddWindow(SettingsWindow.ToSharedRef());
-		}
-	}
-}
-
-void FDerivedDataEditorModule::OnSettingsDialogClosed(const TSharedRef<SWindow>& InWindow)
-{
-	SettingsWindow = nullptr;
-	SettingsDialog = nullptr;
-}
-
 #undef LOCTEXT_NAMESPACE
diff --git a/Engine/Source/Editor/DerivedDataEditor/Private/DerivedDataInformation.cpp b/Engine/Source/Editor/DerivedDataEditor/Private/DerivedDataInformation.cpp
--- a/Engine/Source/Editor/DerivedDataEditor/Private/DerivedDataInformation.cpp
+++ b/Engine/Source/Editor/DerivedDataEditor/Private/DerivedDataInformation.cpp
@@ -107,11 +107,11 @@ double FDerivedDataInformation::GetCacheActivityTimeSeconds(bool bGet, bool bLoc
 bool FDerivedDataInformation::GetHasRemoteCache()
 {
 	for (const TSharedRef<const FDerivedDataCacheStatsNode>& Usage : GetCacheUsageStats())
-		{
-		if (!Usage->IsLocal())
 	{
-			return true;
-	}
+		if (!Usage->IsLocal())
+		{
+				return true;
+		}
 	}
 	return false;
 }
@@ -119,11 +119,11 @@ bool FDerivedDataInformation::GetHasRemoteCache()
 bool FDerivedDataInformation::GetHasZenCache()
 {
 	for (const TSharedRef<const FDerivedDataCacheStatsNode>& Usage : GetCacheUsageStats())
-		{
-		if (Usage->GetCacheType().Equals(TEXT("Zen")))
 	{
-			return true;
-	}
+		if (Usage->GetCacheType().Equals(TEXT("Zen")))
+		{
+				return true;
+		}
 	}
 	return false;
 }
@@ -131,11 +131,11 @@ bool FDerivedDataInformation::GetHasZenCache()
 bool FDerivedDataInformation::GetHasUnrealCloudCache()
 {
 	for (const TSharedRef<const FDerivedDataCacheStatsNode>& Usage : GetCacheUsageStats())
-		{
-		if (Usage->GetCacheType().Equals(TEXT("Unreal Cloud DDC")))
 	{
-			return true;
-	}
+		if (Usage->GetCacheType().Equals(TEXT("Unreal Cloud DDC")))
+		{
+				return true;
+		}
 	}
 	return false;
 }
@@ -174,30 +174,30 @@ void FDerivedDataInformation::UpdateRemoteCacheState()
 
 		if (EditorSettings && EditorSettings->bEnableDDCNotifications)
 		{
-			if (DDCProjectSettings->RecommendEveryoneUseUnrealCloudDDC && EditorSettings->bNotifyUseUnrealCloudDDC && !GetHasZenCache() && !GetHasUnrealCloudCache() && (FCString::Stricmp(GetDerivedDataCache()->GetGraphName(), TEXT("NoJupiter")) != 0))
+			if (DDCProjectSettings->RecommendEveryoneUseUnrealCloudDDC && EditorSettings->bNotifyUseUnrealCloudDDC && !GetHasZenCache() && !GetHasUnrealCloudCache())
 			{
 				RemoteCacheState = ERemoteCacheState::Warning;
-				RemoteCacheWarningMessage = FText(LOCTEXT("UnrealCloudDDCWarning", "It is recommended that you use a DDC graph that supports Unreal Cloud DDC. Please check any -ddc commandline overrides."));
+				RemoteCacheWarningMessage = FText(LOCTEXT("UnrealCloudDDCWarning", "It is recommended that you use Unreal Cloud DDC.\nDisable this notification in the settings."));
 			}
 			else if (DDCProjectSettings->RecommendEveryoneSetupAGlobalLocalDDCPath && EditorSettings->bNotifySetupDDCPath && EditorSettings->GlobalLocalDDCPath.Path.IsEmpty())
 			{
 				RemoteCacheState = ERemoteCacheState::Warning;
-				RemoteCacheWarningMessage = FText(LOCTEXT("GlobalLocalDDCPathWarning", "It is recommended that you set up a valid Global Local DDC Path"));
+				RemoteCacheWarningMessage = FText(LOCTEXT("GlobalLocalDDCPathWarning", "It is recommended that you set up a valid Global Local DDC Path.\nDisable this notification or set up a valid Global Local DDC Path in the settings."));
 			}
 			else if (DDCProjectSettings->RecommendEveryoneSetupAGlobalSharedDDCPath && EditorSettings->bNotifySetupDDCPath && EditorSettings->GlobalSharedDDCPath.Path.IsEmpty())
 			{
 				RemoteCacheState = ERemoteCacheState::Warning;
-				RemoteCacheWarningMessage = FText(LOCTEXT("GlobalSharedDDCPathWarning", "It is recommended that you set up a valid Global Shared DDC Path"));
+				RemoteCacheWarningMessage = FText(LOCTEXT("GlobalSharedDDCPathWarning", "It is recommended that you set up a valid Global Shared DDC Path.\nDisable this notification or set up a valid Global Shared DDC Path in the settings."));
 			}
 			else if (DDCProjectSettings->RecommendEveryoneEnableS3DDC && EditorSettings->bNotifyEnableS3DD && !EditorSettings->bEnableS3DDC)
 			{
 				RemoteCacheState = ERemoteCacheState::Warning;
-				RemoteCacheWarningMessage = FText(LOCTEXT("AWSS3CacheEnabledWarning", "It is recommended that you enable the AWS S3 Cache"));
+				RemoteCacheWarningMessage = FText(LOCTEXT("AWSS3CacheEnabledWarning", "It is recommended that you enable the AWS S3 Cache.\nDisable this notification or enable the AWS S3 Cache in the settings."));
 			}
 			else if (DDCProjectSettings->RecommendEveryoneSetupAGlobalS3DDCPath && EditorSettings->bNotifySetupDDCPath && EditorSettings->GlobalS3DDCPath.Path.IsEmpty())
 			{
 				RemoteCacheState = ERemoteCacheState::Warning;
-				RemoteCacheWarningMessage = FText(LOCTEXT("S3GloblaLocalPathWarning", "It is recommended that you set up a valid Global Local S3 DDC Path"));
+				RemoteCacheWarningMessage = FText(LOCTEXT("S3GloblaLocalPathWarning", "It is recommended that you set up a valid Global Local S3 DDC Path.\nDisable this notification or set up a valid Global Local S3 DDC Path in the settings."));
 			}
 		}
 	}
diff --git a/Engine/Source/Editor/DerivedDataEditor/Private/SDerivedDataCacheSettings.cpp b/Engine/Source/Editor/DerivedDataEditor/Private/SDerivedDataCacheSettings.cpp
deleted file mode 100644
diff --git a/Engine/Source/Editor/DerivedDataEditor/Private/SDerivedDataCacheSettings.h b/Engine/Source/Editor/DerivedDataEditor/Private/SDerivedDataCacheSettings.h
deleted file mode 100644
diff --git a/Engine/Source/Editor/DerivedDataEditor/Private/SDerivedDataStatusBar.cpp b/Engine/Source/Editor/DerivedDataEditor/Private/SDerivedDataStatusBar.cpp
--- a/Engine/Source/Editor/DerivedDataEditor/Private/SDerivedDataStatusBar.cpp
+++ b/Engine/Source/Editor/DerivedDataEditor/Private/SDerivedDataStatusBar.cpp
@@ -268,7 +268,7 @@ EActiveTimerReturnType SDerivedDataStatusBarWidget::UpdateWarnings(double InCurr
 			Info.FadeOutDuration = 0.0f;
 			Info.ExpireDuration = 0.0f;
 
-			Info.ButtonDetails.Add(FNotificationButtonInfo(LOCTEXT("UpdateSettings", "Update Settings"), FText(), FSimpleDelegate::CreateLambda([NotificationFuture = NotificationPromise.GetFuture().Share()]() {
+			Info.ButtonDetails.Add(FNotificationButtonInfo(LOCTEXT("OpenSettings", "Open Settings"), FText(), FSimpleDelegate::CreateLambda([NotificationFuture = NotificationPromise.GetFuture().Share()]() {
 				FModuleManager::LoadModuleChecked<ISettingsModule>("Settings").ShowViewer("Editor", "General", "Global");
 
 				TWeakPtr<SNotificationItem> NotificationPtr = NotificationFuture.Get();
diff --git a/Engine/Source/Editor/DerivedDataEditor/Public/DerivedDataEditorModule.h b/Engine/Source/Editor/DerivedDataEditor/Public/DerivedDataEditorModule.h
--- a/Engine/Source/Editor/DerivedDataEditor/Public/DerivedDataEditorModule.h
+++ b/Engine/Source/Editor/DerivedDataEditor/Public/DerivedDataEditorModule.h
@@ -38,8 +38,6 @@ public:
 	void ShowResourceUsageTab();
 	void ShowCacheStatisticsTab();
 	
-	void ShowSettingsDialog();
-
 private:
 
 	TSharedPtr<SWidget> CreateResourceUsageDialog();
@@ -48,8 +46,6 @@ private:
 	TSharedRef<SDockTab> CreateResourceUsageTab(const FSpawnTabArgs& Args);
 	TSharedRef<SDockTab> CreateCacheStatisticsTab(const FSpawnTabArgs& Args);
 
-	void OnSettingsDialogClosed(const TSharedRef<SWindow>& InWindow);
-
 	TWeakPtr<SDockTab> ResourceUsageTab;
 	TWeakPtr<SDockTab> CacheStatisticsTab;
 
diff --git a/Engine/Source/Editor/UnrealEd/Classes/Editor/EditorPerformanceSettings.h b/Engine/Source/Editor/UnrealEd/Classes/Editor/EditorPerformanceSettings.h
--- a/Engine/Source/Editor/UnrealEd/Classes/Editor/EditorPerformanceSettings.h
+++ b/Engine/Source/Editor/UnrealEd/Classes/Editor/EditorPerformanceSettings.h
@@ -50,7 +50,7 @@ class UEditorPerformanceSettings : public UDeveloperSettings
 	uint32 bMonitorEditorPerformance:1;
 
 	/** When enabled, Shared Data Cache performance notifications may be displayed when not connected to a shared cache */
-	UPROPERTY(EditAnywhere, config, Category = EditorPerformance, meta = (DisplayName = "Enable Shared Data Cache Performance Notifications"))
+	UE_DEPRECATED(5.2, "No longer supported")
 	uint32 bEnableSharedDDCPerformanceNotifications : 1;
 
 	/** When enabled, a warning will appear in the viewport when your editors scalability settings are non-default and you may be viewing a low quality scene */
diff --git a/Engine/Source/Editor/UnrealEd/Private/EditorPerformanceSettings.cpp b/Engine/Source/Editor/UnrealEd/Private/EditorPerformanceSettings.cpp
--- a/Engine/Source/Editor/UnrealEd/Private/EditorPerformanceSettings.cpp
+++ b/Engine/Source/Editor/UnrealEd/Private/EditorPerformanceSettings.cpp
@@ -19,7 +19,6 @@ UEditorPerformanceSettings::UEditorPerformanceSettings(const FObjectInitializer&
 	, bThrottleCPUWhenNotForeground(true)
 	, bDisableRealtimeViewportsInRemoteSessions(true)
 	, bMonitorEditorPerformance(false)
-	, bEnableSharedDDCPerformanceNotifications(true)
 	, bEnableScalabilityWarningIndicator(true)
 	, bDisplayHighDPIViewports(true)
 	, bOverridePIEScreenPercentage(true)
