Minor tweak to analytics code for inconsistent naming of Editor load events.
Fixed Editor notifcation warning for when Unreal Cloud DDC is not available. Without this notification developers could unknowingly experience poor Editor performance.

diff --git a/Engine/Source/Editor/DerivedDataEditor/Private/DerivedDataInformation.cpp b/Engine/Source/Editor/DerivedDataEditor/Private/DerivedDataInformation.cpp
--- a/Engine/Source/Editor/DerivedDataEditor/Private/DerivedDataInformation.cpp
+++ b/Engine/Source/Editor/DerivedDataEditor/Private/DerivedDataInformation.cpp
@@ -134,7 +134,7 @@ bool FDerivedDataInformation::GetHasUnrealCloudCache()
 	{
 		if (Usage->GetCacheType().Equals(TEXT("Unreal Cloud DDC")))
 		{
-				return true;
+			return true;
 		}
 	}
 	return false;
@@ -174,7 +174,7 @@ void FDerivedDataInformation::UpdateRemoteCacheState()
 
 		if (EditorSettings && EditorSettings->bEnableDDCNotifications)
 		{
-			if (DDCProjectSettings->RecommendEveryoneUseUnrealCloudDDC && EditorSettings->bNotifyUseUnrealCloudDDC && !GetHasZenCache() && !GetHasUnrealCloudCache())
+			if (DDCProjectSettings->RecommendEveryoneUseUnrealCloudDDC && EditorSettings->bNotifyUseUnrealCloudDDC && !GetHasUnrealCloudCache())
 			{
 				RemoteCacheState = ERemoteCacheState::Warning;
 				RemoteCacheWarningMessage = FText(LOCTEXT("UnrealCloudDDCWarning", "It is recommended that you use Unreal Cloud DDC.\nDisable this notification in the settings."));
