Changed Derived Data Cache Usage UI to use the same Resource Stats as Studio Telemetry
The total is already served up as an Asset Type so had to and some logic to separate it in the view.
Changed sorting from Size to Count. I felt that sorting by highest count is more useful as is it clearly indicates the less efficient DDC work.
Added Hit Efficiency to Asset Stats and added this as an extra column to the UI

diff --git a/Engine/Plugins/Experimental/StudioTelemetry/Source/StudioTelemetry/Private/StudioTelemetryEditor.cpp b/Engine/Plugins/Experimental/StudioTelemetry/Source/StudioTelemetry/Private/StudioTelemetryEditor.cpp
--- a/Engine/Plugins/Experimental/StudioTelemetry/Source/StudioTelemetry/Private/StudioTelemetryEditor.cpp
+++ b/Engine/Plugins/Experimental/StudioTelemetry/Source/StudioTelemetry/Private/StudioTelemetryEditor.cpp
@@ -214,11 +214,10 @@ void FStudioTelemetryEditor::RecordEvent_DDCResource(const FString& Context, TAr
 	// Send a resource event per asset type
 	for (const FDerivedDataCacheResourceStat& Stat : ResourceStats)
 	{
-		const int64 TotalCount = Stat.BuildCount + Stat.LoadCount;
 		const double TotalTimeSec = Stat.BuildTimeSec + Stat.LoadTimeSec;
 		const int64 TotalSizeMB = Stat.BuildSizeMB + Stat.LoadSizeMB;
 
-		if (Stat.AssetType.IsEmpty() || TotalCount==0)
+		if (Stat.AssetType.IsEmpty() || Stat.TotalCount==0)
 		{
 			// Empty asset type or nothing was built or loaded for this type
 			continue;
@@ -233,10 +232,10 @@ void FStudioTelemetryEditor::RecordEvent_DDCResource(const FString& Context, TAr
 		EventAttributes.Emplace(TEXT("Build_Count"), Stat.BuildCount);
 		EventAttributes.Emplace(TEXT("Build_TimeSec"), Stat.BuildTimeSec);
 		EventAttributes.Emplace(TEXT("Build_SizeMB"), Stat.BuildSizeMB);
-		EventAttributes.Emplace(TEXT("Total_Count"), TotalCount);
+		EventAttributes.Emplace(TEXT("Total_Count"), Stat.TotalCount);
 		EventAttributes.Emplace(TEXT("Total_TimeSec"), TotalTimeSec);
 		EventAttributes.Emplace(TEXT("Total_SizeMB"), TotalSizeMB);
-		EventAttributes.Emplace(TEXT("Efficiency"), double(Stat.LoadCount)/double(TotalCount) );
+		EventAttributes.Emplace(TEXT("Efficiency"), Stat.Efficiency);
 		EventAttributes.Emplace(TEXT("Thread_TimeSec"), Stat.GameThreadTimeSec);
 
 		FStudioTelemetry::Get().RecordEvent(TEXT("Core.DDC.Resource"), EventAttributes);
@@ -635,11 +634,13 @@ void FStudioTelemetryEditor::Initialize()
 			Attributes.Emplace(TEXT("MapName"), EditorMapName);
 
 			PIESpan->AddAttributes(Attributes);
+			PIEStartupSpan->AddAttributes(Attributes);
 		});
 
 	FEditorDelegates::PreBeginPIE.AddLambda([this](bool)
 		{	
 			PIEPreBeginSpan = FStudioTelemetry::Get().StartSpan(PIEPreBeginSpanName, PIEStartupSpan);
+			PIEPreBeginSpan->AddAttributes(PIESpan->GetAttributes());
 		});
 
 	FEditorDelegates::BeginPIE.AddLambda([this](bool)
@@ -700,6 +701,7 @@ void FStudioTelemetryEditor::Initialize()
 				}
 
 				PIEInteractSpan = FStudioTelemetry::Get().StartSpan(PIEInteractSpanName, PIESpan);
+				PIEInteractSpan->AddAttributes(PIESpan->GetAttributes());
 			}
 		});
 
@@ -710,6 +712,7 @@ void FStudioTelemetryEditor::Initialize()
 				// PIE is ending so no longer interactive
 				FStudioTelemetry::Get().EndSpan(PIEInteractSpan);
 				PIEShutdownSpan = FStudioTelemetry::Get().StartSpan(PIEShutdownSpanName, PIESpan);
+				PIEShutdownSpan->AddAttributes(PIESpan->GetAttributes());
 			}
 		});
 
diff --git a/Engine/Source/Developer/DerivedDataCache/Public/DerivedDataCacheUsageStats.h b/Engine/Source/Developer/DerivedDataCache/Public/DerivedDataCacheUsageStats.h
--- a/Engine/Source/Developer/DerivedDataCache/Public/DerivedDataCacheUsageStats.h
+++ b/Engine/Source/Developer/DerivedDataCache/Public/DerivedDataCacheUsageStats.h
@@ -233,7 +233,10 @@ public:
 		BuildSizeMB(InBuildSizeMB),
 		BuildCount(InAssetsBuilt),
 		GameThreadTimeSec(bIsGameThreadTime ? InLoadTimeSec + InBuildTimeSec : 0.0)
-	{}
+	{
+		TotalCount = LoadCount + BuildCount;
+		Efficiency = TotalCount>0? (double)LoadCount / (double)TotalCount : 1.0;
+	}
 
 	const FDerivedDataCacheResourceStat& operator+(const FDerivedDataCacheResourceStat& OtherStat)
 	{
@@ -247,6 +250,9 @@ public:
 		BuildTimeSec += OtherStat.BuildTimeSec;
 		BuildSizeMB += OtherStat.BuildSizeMB;
 
+		TotalCount = LoadCount + BuildCount;
+		Efficiency = TotalCount > 0 ? (double)LoadCount / (double)TotalCount : 1.0;
+
 		return *this;
 	}
 
@@ -262,6 +268,9 @@ public:
 		BuildTimeSec -= OtherStat.BuildTimeSec;
 		BuildSizeMB -= OtherStat.BuildSizeMB;
 
+		TotalCount = LoadCount + BuildCount;
+		Efficiency = TotalCount > 0 ? (double)LoadCount / (double)TotalCount : 1.0;
+
 		return *this;
 	}
 
@@ -279,13 +288,16 @@ public:
 
 	FString AssetType;
 
-	double LoadTimeSec;
-	double LoadSizeMB;
-	int64 LoadCount;
+	double LoadTimeSec = 0;
+	double LoadSizeMB = 0;
+	int64 LoadCount = 0;
+
+	double BuildTimeSec = 0;
+	double BuildSizeMB = 0;
+	int64 BuildCount = 0;
 
-	double BuildTimeSec;
-	double BuildSizeMB;
-	int64 BuildCount;
+	int64 TotalCount = 0;
+	double Efficiency = 0;
 
 	double GameThreadTimeSec;
 };
diff --git a/Engine/Source/Editor/DerivedDataEditor/Private/SDerivedDataDialogs.cpp b/Engine/Source/Editor/DerivedDataEditor/Private/SDerivedDataDialogs.cpp
--- a/Engine/Source/Editor/DerivedDataEditor/Private/SDerivedDataDialogs.cpp
+++ b/Engine/Source/Editor/DerivedDataEditor/Private/SDerivedDataDialogs.cpp
@@ -69,22 +69,22 @@ EActiveTimerReturnType SDerivedDataRemoteStoreDialog::UpdateGridPanels(double In
 
 TSharedRef<SWidget> SDerivedDataRemoteStoreDialog::GetGridPanel()
 {
+	// Gather the latest resource stats
 	TArray<FDerivedDataCacheResourceStat> DDCResourceStats;
-
-	// Grab the latest resource stats
-	GetDerivedDataCacheRef().GatherResourceStats(DDCResourceStats);
+	GatherDerivedDataCacheResourceStats(DDCResourceStats);
 
 	FDerivedDataCacheResourceStat DDCResourceStatsTotal(TEXT("Total"));
 
-	// Accumulate Totals
+	// Find the total
 	for (const FDerivedDataCacheResourceStat& Stat : DDCResourceStats)
 	{
-		DDCResourceStatsTotal += Stat;
+		if (Stat.AssetType == TEXT("Total"))
+		{
+			DDCResourceStatsTotal = Stat;
+			break;
+		}
 	}
 
-	const int64 TotalCount = DDCResourceStatsTotal.LoadCount + DDCResourceStatsTotal.BuildCount;
-	const double Efficiency = TotalCount > 0 ? static_cast<double>(DDCResourceStatsTotal.LoadCount) / static_cast<double>(TotalCount) : 0.0;
-
 	const double DownloadedBytesMB = FUnitConversion::Convert(FDerivedDataInformation::GetCacheActivitySizeBytes(true, false), EUnit::Bytes, EUnit::Megabytes);
 	const double UploadedBytesMB = FUnitConversion::Convert(FDerivedDataInformation::GetCacheActivitySizeBytes(false, false), EUnit::Bytes, EUnit::Megabytes);
 
@@ -184,17 +184,16 @@ EActiveTimerReturnType SDerivedDataResourceUsageDialog::UpdateGridPanels(double
 
 TSharedRef<SWidget> SDerivedDataResourceUsageDialog::GetGridPanel()
 {
+	// Gather the latest resource stats
 	TArray<FDerivedDataCacheResourceStat> DDCResourceStats;
+	GatherDerivedDataCacheResourceStats(DDCResourceStats);
 
-	// Grab the resource stats
-	GetDerivedDataCacheRef().GatherResourceStats(DDCResourceStats);
-
-	// Sort results on descending build size, then descending load size, then ascending asset type.
+	// Sort results on descending build count, then descending load size, then ascending asset type.
 	const auto CompareStats = [](const FDerivedDataCacheResourceStat& A, const FDerivedDataCacheResourceStat& B)
 	{
-		if (A.BuildSizeMB != B.BuildSizeMB)
+		if (A.BuildCount != B.BuildCount)
 		{
-			return A.BuildSizeMB > B.BuildSizeMB;
+			return A.BuildCount > B.BuildCount;
 		}
 		if (A.LoadSizeMB != B.LoadSizeMB)
 		{
@@ -206,10 +205,14 @@ TSharedRef<SWidget> SDerivedDataResourceUsageDialog::GetGridPanel()
 
 	FDerivedDataCacheResourceStat DDCResourceStatsTotal(TEXT("Total"));
 
-	// Accumulate Totals
+	// Find the total and store it off, we will render that out separately
 	for (const FDerivedDataCacheResourceStat& Stat : DDCResourceStats)
 	{
-		DDCResourceStatsTotal += Stat;
+		if (Stat.AssetType == TEXT("Total"))
+		{
+			DDCResourceStatsTotal = Stat;
+			break;
+		}
 	}
 
 	TSharedRef<SGridPanel> Panel =
@@ -226,7 +229,7 @@ TSharedRef<SWidget> SDerivedDataResourceUsageDialog::GetGridPanel()
 	const FMargin DefaultMargin(0.0f, RowMargin, ColumnMargin, RowMargin);
 	const FMargin DefaultMarginFirstColumn(ColumnMargin, RowMargin);
 
-	Panel->AddSlot(2, Row)
+	Panel->AddSlot(3, Row)
 	.HAlign(HAlign_Center)
 	[
 		SNew(STextBlock)
@@ -236,7 +239,7 @@ TSharedRef<SWidget> SDerivedDataResourceUsageDialog::GetGridPanel()
 		.Text(LOCTEXT("Loaded", "Loaded"))
 	];
 
-	Panel->AddSlot(5, Row)
+	Panel->AddSlot(6, Row)
 	.HAlign(HAlign_Center)
 	[
 		SNew(STextBlock)
@@ -259,6 +262,16 @@ TSharedRef<SWidget> SDerivedDataResourceUsageDialog::GetGridPanel()
 
 	Panel->AddSlot(1, Row)
 	.HAlign(HAlign_Right)
+	[
+		SNew(STextBlock)
+		.Margin(TitleMargin)
+	.ColorAndOpacity(TitleColor)
+	.Font(TitleFont)
+	.Text(LOCTEXT("Hit", "Hit%"))
+	];
+
+	Panel->AddSlot(2, Row)
+	.HAlign(HAlign_Right)
 	[
 		SNew(STextBlock)
 		.Margin(TitleMargin)
@@ -267,7 +280,7 @@ TSharedRef<SWidget> SDerivedDataResourceUsageDialog::GetGridPanel()
 		.Text(LOCTEXT("Count", "Count"))
 	];
 
-	Panel->AddSlot(2, Row)
+	Panel->AddSlot(3, Row)
 	.HAlign(HAlign_Right)
 	[
 		SNew(STextBlock)
@@ -277,7 +290,7 @@ TSharedRef<SWidget> SDerivedDataResourceUsageDialog::GetGridPanel()
 		.Text(LOCTEXT("Time (Sec)", "Time (Sec)"))
 	];
 
-	Panel->AddSlot(3, Row)
+	Panel->AddSlot(4, Row)
 	.HAlign(HAlign_Right)
 	[
 		SNew(STextBlock)
@@ -287,7 +300,7 @@ TSharedRef<SWidget> SDerivedDataResourceUsageDialog::GetGridPanel()
 		.Text(LOCTEXT("Size (MiB)", "Size (MiB)"))
 	];
 
-	Panel->AddSlot(4, Row)
+	Panel->AddSlot(5, Row)
 	.HAlign(HAlign_Right)
 	[
 		SNew(STextBlock)
@@ -297,7 +310,7 @@ TSharedRef<SWidget> SDerivedDataResourceUsageDialog::GetGridPanel()
 		.Text(LOCTEXT("Count", "Count"))
 	];
 
-	Panel->AddSlot(5, Row)
+	Panel->AddSlot(6, Row)
 	.HAlign(HAlign_Right)
 	[
 		SNew(STextBlock)
@@ -307,7 +320,7 @@ TSharedRef<SWidget> SDerivedDataResourceUsageDialog::GetGridPanel()
 		.Text(LOCTEXT("Time (Sec)", "Time (Sec)"))
 	];
 
-	Panel->AddSlot(6, Row)
+	Panel->AddSlot(7, Row)
 	.HAlign(HAlign_Right)
 	[
 		SNew(STextBlock)
@@ -321,6 +334,11 @@ TSharedRef<SWidget> SDerivedDataResourceUsageDialog::GetGridPanel()
 
 	for (const FDerivedDataCacheResourceStat& Stat : DDCResourceStats)
 	{	
+		if (Stat.AssetType == TEXT("Total"))
+		{
+			continue;
+		}
+
 		Panel->AddSlot(0, Row)
 		[
 			SNew(STextBlock)
@@ -333,15 +351,16 @@ TSharedRef<SWidget> SDerivedDataResourceUsageDialog::GetGridPanel()
 		[
 			SNew(STextBlock)
 			.Margin(DefaultMargin)
-			.Text(FText::FromString(ZeroDecimalFormat(Stat.LoadCount)))
+			.Text(FText::FromString(SingleDecimalFormat(Stat.Efficiency*100.0)))
 		];
 
+
 		Panel->AddSlot(2, Row)
 		.HAlign(HAlign_Right)
 		[
 			SNew(STextBlock)
 			.Margin(DefaultMargin)
-			.Text(FText::FromString(SingleDecimalFormat(Stat.LoadTimeSec)))
+			.Text(FText::FromString(ZeroDecimalFormat(Stat.LoadCount)))
 		];
 
 		Panel->AddSlot(3, Row)
@@ -349,7 +368,7 @@ TSharedRef<SWidget> SDerivedDataResourceUsageDialog::GetGridPanel()
 		[
 			SNew(STextBlock)
 			.Margin(DefaultMargin)
-			.Text(FText::FromString(SingleDecimalFormat(Stat.LoadSizeMB)))
+			.Text(FText::FromString(SingleDecimalFormat(Stat.LoadTimeSec)))
 		];
 
 		Panel->AddSlot(4, Row)
@@ -357,7 +376,7 @@ TSharedRef<SWidget> SDerivedDataResourceUsageDialog::GetGridPanel()
 		[
 			SNew(STextBlock)
 			.Margin(DefaultMargin)
-			.Text(FText::FromString(ZeroDecimalFormat(Stat.BuildCount)))
+			.Text(FText::FromString(SingleDecimalFormat(Stat.LoadSizeMB)))
 		];
 
 		Panel->AddSlot(5, Row)
@@ -365,11 +384,19 @@ TSharedRef<SWidget> SDerivedDataResourceUsageDialog::GetGridPanel()
 		[
 			SNew(STextBlock)
 			.Margin(DefaultMargin)
-			.Text(FText::FromString(SingleDecimalFormat(Stat.BuildTimeSec)))
+			.Text(FText::FromString(ZeroDecimalFormat(Stat.BuildCount)))
 		];
 
 		Panel->AddSlot(6, Row)
 		.HAlign(HAlign_Right)
+		[
+			SNew(STextBlock)
+			.Margin(DefaultMargin)
+			.Text(FText::FromString(SingleDecimalFormat(Stat.BuildTimeSec)))
+		];
+
+		Panel->AddSlot(7, Row)
+		.HAlign(HAlign_Right)
 		[
 			SNew(STextBlock)
 			.Margin(DefaultMargin)
@@ -395,9 +422,10 @@ TSharedRef<SWidget> SDerivedDataResourceUsageDialog::GetGridPanel()
 		.Margin(TitleMargin)
 		.ColorAndOpacity(TitleColor)
 		.Font(TitleFont)
-		.Text(FText::FromString(ZeroDecimalFormat(DDCResourceStatsTotal.LoadCount)))
+		.Text(FText::FromString(SingleDecimalFormat(DDCResourceStatsTotal.Efficiency * 100.0)))
 	];
 
+
 	Panel->AddSlot(2, Row)
 	.HAlign(HAlign_Right)
 	[
@@ -405,7 +433,7 @@ TSharedRef<SWidget> SDerivedDataResourceUsageDialog::GetGridPanel()
 		.Margin(TitleMargin)
 		.ColorAndOpacity(TitleColor)
 		.Font(TitleFont)
-		.Text(FText::FromString(SingleDecimalFormat(DDCResourceStatsTotal.LoadTimeSec)))
+		.Text(FText::FromString(ZeroDecimalFormat(DDCResourceStatsTotal.LoadCount)))
 	];
 
 	Panel->AddSlot(3, Row)
@@ -415,7 +443,7 @@ TSharedRef<SWidget> SDerivedDataResourceUsageDialog::GetGridPanel()
 		.Margin(TitleMargin)
 		.ColorAndOpacity(TitleColor)
 		.Font(TitleFont)
-		.Text(FText::FromString(SingleDecimalFormat(DDCResourceStatsTotal.LoadSizeMB)))
+		.Text(FText::FromString(SingleDecimalFormat(DDCResourceStatsTotal.LoadTimeSec)))
 	];
 
 	Panel->AddSlot(4, Row)
@@ -425,7 +453,7 @@ TSharedRef<SWidget> SDerivedDataResourceUsageDialog::GetGridPanel()
 		.Margin(TitleMargin)
 		.ColorAndOpacity(TitleColor)
 		.Font(TitleFont)
-		.Text(FText::FromString(ZeroDecimalFormat(DDCResourceStatsTotal.BuildCount)))
+		.Text(FText::FromString(SingleDecimalFormat(DDCResourceStatsTotal.LoadSizeMB)))
 	];
 
 	Panel->AddSlot(5, Row)
@@ -435,11 +463,21 @@ TSharedRef<SWidget> SDerivedDataResourceUsageDialog::GetGridPanel()
 		.Margin(TitleMargin)
 		.ColorAndOpacity(TitleColor)
 		.Font(TitleFont)
-		.Text(FText::FromString(SingleDecimalFormat(DDCResourceStatsTotal.BuildTimeSec)))
+		.Text(FText::FromString(ZeroDecimalFormat(DDCResourceStatsTotal.BuildCount)))
 	];
 
 	Panel->AddSlot(6, Row)
 	.HAlign(HAlign_Right)
+	[
+		SNew(STextBlock)
+		.Margin(TitleMargin)
+		.ColorAndOpacity(TitleColor)
+		.Font(TitleFont)
+		.Text(FText::FromString(SingleDecimalFormat(DDCResourceStatsTotal.BuildTimeSec)))
+	];
+
+	Panel->AddSlot(7, Row)
+	.HAlign(HAlign_Right)
 	[
 		SNew(STextBlock)
 		.Margin(TitleMargin)
