Change text wrapping size limit on DDC stats panel that would lead to flickering (possibly due to numeric imprecision in the computations when doing text flow).

diff --git a/Engine/Source/Editor/DerivedDataEditor/Private/SDerivedDataDialogs.cpp b/Engine/Source/Editor/DerivedDataEditor/Private/SDerivedDataDialogs.cpp
--- a/Engine/Source/Editor/DerivedDataEditor/Private/SDerivedDataDialogs.cpp
+++ b/Engine/Source/Editor/DerivedDataEditor/Private/SDerivedDataDialogs.cpp
@@ -653,7 +653,7 @@ TSharedRef<SWidget> SDerivedDataCacheStatisticsDialog::GetGridPanel()
 		.Font(TitleFont)
 		.Text(LOCTEXT("Read Speed", "Read Speed (MiB/s)"))
 		.AutoWrapText(true)
-		.WrapTextAt(66.0f)
+		.WrapTextAt(100.0f)
 	];
 
 	Panel->AddSlot(8, Row)
@@ -665,7 +665,7 @@ TSharedRef<SWidget> SDerivedDataCacheStatisticsDialog::GetGridPanel()
 		.Font(TitleFont)
 		.Text(LOCTEXT("Write Speed", "Write Speed (MiB/s)"))
 		.AutoWrapText(true)
-		.WrapTextAt(66.0f)
+		.WrapTextAt(100.0f)
 	];
 
 	Panel->AddSlot(9, Row)
